[RemoteControlComponents] Moved plugin

diff --git a/Engine/Plugins/Experimental/RemoteControlComponents/RemoteControlComponents.uplugin b/Engine/Plugins/Experimental/RemoteControlComponents/RemoteControlComponents.uplugin
new file mode 100644
--- /dev/null
+++ b/Engine/Plugins/Experimental/RemoteControlComponents/RemoteControlComponents.uplugin
@@ -0,0 +1,33 @@
+﻿{
+	"FileVersion": 3,
+	"Version": 0.1,
+	"VersionName": "0.1.0-alpha",
+	"FriendlyName": "RemoteControlComponents",
+	"Description": "",
+	"Category": "Experimental",
+	"CreatedBy": "Epic Games, Inc.",
+	"CreatedByURL": "http://epicgames.com",
+	"DocsURL": "",
+	"MarketplaceURL": "",
+	"CanContainContent": true,
+	"IsBetaVersion": true,
+	"Installed": false,
+	"Modules": [
+		{
+			"Name": "RemoteControlComponents",
+			"Type": "Runtime",
+			"LoadingPhase": "Default"
+		},
+		{
+			"Name": "RemoteControlComponentsEditor",
+			"Type": "Editor",
+			"LoadingPhase": "Default"
+		}
+	],
+	"Plugins": [
+		{
+			"Name": "RemoteControl",
+			"Enabled": true
+		}
+	]
+}
diff --git a/Engine/Plugins/Experimental/RemoteControlComponents/Source/RemoteControlComponents/Private/RemoteControlComponentsContext.cpp b/Engine/Plugins/Experimental/RemoteControlComponents/Source/RemoteControlComponents/Private/RemoteControlComponentsContext.cpp
new file mode 100644
--- /dev/null
+++ b/Engine/Plugins/Experimental/RemoteControlComponents/Source/RemoteControlComponents/Private/RemoteControlComponentsContext.cpp
@@ -0,0 +1,76 @@
+﻿// Copyright Epic Games, Inc. All Rights Reserved.
+
+#include "RemoteControlComponentsContext.h"
+#include "Engine/World.h"
+#include "GameFramework/Actor.h"
+#include "RemoteControlPreset.h"
+#include "RemoteControlTrackerComponent.h"
+
+FRemoteControlComponentsContext::FRemoteControlComponentsContext(UWorld* InWorld, URemoteControlPreset* InPreset)
+	: WorldWeak(InWorld)
+	, PresetWeak(InPreset)
+{
+}
+
+bool FRemoteControlComponentsContext::RegisterActor(AActor* InActor)
+{
+	if (CanRegisterActor(InActor))
+	{
+		TrackedActors.Add(InActor);
+		return true;
+	}
+
+	return false;
+}
+
+bool FRemoteControlComponentsContext::UnregisterActor(AActor* InActor)
+{
+	return TrackedActors.Remove(InActor) > 0;
+}
+
+bool FRemoteControlComponentsContext::IsActorRegistered(const AActor* InActor) const
+{
+	return InActor && TrackedActors.Contains(InActor);
+}
+
+bool FRemoteControlComponentsContext::CanRegisterActor(const AActor* InActor) const
+{
+	if (!InActor)
+	{
+		return false;
+	}
+
+	if (IsActorRegistered(InActor))
+	{
+		return false;
+	}
+	
+	const UWorld* ActorWorld = InActor->GetWorld();
+	const UWorld* World = GetWorld();
+	if (!ActorWorld || ActorWorld != World)
+	{
+		return false;
+	}
+
+	return !!InActor->FindComponentByClass<URemoteControlTrackerComponent>();
+}
+
+bool FRemoteControlComponentsContext::IsValid() const
+{
+	return WorldWeak.IsValid() && PresetWeak.IsValid();
+}
+
+bool FRemoteControlComponentsContext::HasTrackedActors() const
+{
+	return !TrackedActors.IsEmpty();
+}
+
+UWorld* FRemoteControlComponentsContext::GetWorld() const
+{
+	return WorldWeak.Get();
+}
+
+URemoteControlPreset* FRemoteControlComponentsContext::GetPreset() const
+{
+	return PresetWeak.Get();
+}
diff --git a/Engine/Plugins/Experimental/RemoteControlComponents/Source/RemoteControlComponents/Private/RemoteControlComponentsModule.cpp b/Engine/Plugins/Experimental/RemoteControlComponents/Source/RemoteControlComponents/Private/RemoteControlComponentsModule.cpp
new file mode 100644
--- /dev/null
+++ b/Engine/Plugins/Experimental/RemoteControlComponents/Source/RemoteControlComponents/Private/RemoteControlComponentsModule.cpp
@@ -0,0 +1,5 @@
+// Copyright Epic Games, Inc. All Rights Reserved.
+
+#include "Modules/ModuleManager.h"
+
+IMPLEMENT_MODULE(FDefaultModuleImpl, RemoteControlComponents)
diff --git a/Engine/Plugins/Experimental/RemoteControlComponents/Source/RemoteControlComponents/Private/RemoteControlComponentsUtils.cpp b/Engine/Plugins/Experimental/RemoteControlComponents/Source/RemoteControlComponents/Private/RemoteControlComponentsUtils.cpp
new file mode 100644
--- /dev/null
+++ b/Engine/Plugins/Experimental/RemoteControlComponents/Source/RemoteControlComponents/Private/RemoteControlComponentsUtils.cpp
@@ -0,0 +1,366 @@
+﻿// Copyright Epic Games, Inc. All Rights Reserved.
+
+#include "RemoteControlComponentsUtils.h"
+#include "RemoteControlBinding.h"
+#include "RemoteControlField.h"
+#include "RemoteControlPreset.h"
+#include "RemoteControlTrackerComponent.h"
+#include "Subsystems/RemoteControlComponentsSubsystem.h"
+
+#if WITH_EDITOR
+#include "ScopedTransaction.h"
+#endif
+
+#define LOCTEXT_NAMESPACE "RemoteControlComponentsUtils"
+
+FGuid FRemoteControlComponentsUtils::GetPresetPropertyGuid(URemoteControlPreset* InPreset, UObject* InOwnerObject, const FProperty* InProperty)
+{
+	const TArray<TWeakPtr<FRemoteControlProperty>>& RemoteControlExposedProperties = InPreset->GetExposedEntities<FRemoteControlProperty>();
+	for (const TWeakPtr<FRemoteControlProperty>& RCPropertyWeak : RemoteControlExposedProperties)
+	{
+		if (const TSharedPtr<FRemoteControlProperty>& RCProperty = RCPropertyWeak.Pin())
+		{
+			const bool bIsExposed = InProperty->GetFName() == RCProperty->FieldName
+				&& RCProperty->GetBindings().ContainsByPredicate([&InOwnerObject](const TWeakObjectPtr<URemoteControlBinding>& InBinding)
+				{
+					return InBinding.IsValid() && InBinding->Resolve() == InOwnerObject;
+				});
+
+			if (bIsExposed)
+			{
+				return RCProperty->GetId();
+			}
+		}
+	}
+
+	return FGuid();
+}
+
+URemoteControlTrackerComponent* FRemoteControlComponentsUtils::GetTrackerComponent(AActor* InActor, bool bInAddTrackerIfMissing /* = false*/)
+{
+	if (!InActor)
+	{
+		return nullptr;
+	}
+	
+	if (!InActor->FindComponentByClass<URemoteControlTrackerComponent>() && bInAddTrackerIfMissing)
+	{
+		constexpr bool bShouldTransact = false;
+		FRemoteControlComponentsUtils::AddTrackerComponent({InActor}, bShouldTransact);
+	}
+
+	URemoteControlTrackerComponent* TrackerComponent = InActor->GetComponentByClass<URemoteControlTrackerComponent>();
+
+	return TrackerComponent;
+}
+
+URemoteControlPreset* FRemoteControlComponentsUtils::GetCurrentPreset(const UObject* InObject)
+{
+	if (const URemoteControlComponentsSubsystem* RemoteControlComponentsSubsystem = URemoteControlComponentsSubsystem::Get())
+	{
+		return RemoteControlComponentsSubsystem->GetRegisteredPreset(InObject);
+	}
+
+	return nullptr;
+}
+
+URemoteControlPreset* FRemoteControlComponentsUtils::GetCurrentPreset(const UWorld* InWorld)
+{
+	if (const URemoteControlComponentsSubsystem* RemoteControlComponentsSubsystem = URemoteControlComponentsSubsystem::Get())
+	{
+		return RemoteControlComponentsSubsystem->GetRegisteredPreset(InWorld);
+	}
+
+	return nullptr;
+}
+
+void FRemoteControlComponentsUtils::AddTrackerComponent(const TSet<TWeakObjectPtr<AActor>>& InActors, bool bInShouldTransact)
+{
+#if WITH_EDITOR
+	TSharedPtr<FScopedTransaction> AddTrackerTransaction;
+	if (bInShouldTransact)
+	{
+		AddTrackerTransaction = MakeShared<FScopedTransaction>(LOCTEXT("AddTrackerComponents", "Add Tracker Component(s)"));
+	}
+#endif
+	
+	for (TWeakObjectPtr<AActor> ActorWeak : InActors)
+	{
+		if (ActorWeak.IsValid())
+		{
+			if (AActor* Actor = ActorWeak.Get())
+			{
+				if (Actor->FindComponentByClass<URemoteControlTrackerComponent>())
+				{
+					continue;
+				}
+#if WITH_EDITOR
+				Actor->Modify();
+#endif
+
+				URemoteControlTrackerComponent* TrackerComponent = NewObject<URemoteControlTrackerComponent>(Actor
+					, URemoteControlTrackerComponent::StaticClass()
+					, MakeUniqueObjectName(Actor, URemoteControlTrackerComponent::StaticClass(), TEXT("RemoteControlTrackerComponent"))
+					, RF_Transactional);
+
+				Actor->AddInstanceComponent(TrackerComponent);
+				TrackerComponent->RegisterComponent();
+#if WITH_EDITOR
+				Actor->RerunConstructionScripts();
+#endif
+
+				if (URemoteControlComponentsSubsystem* ControlComponentsSubsystem = URemoteControlComponentsSubsystem::Get())
+				{
+#if WITH_EDITOR
+					ControlComponentsSubsystem->Modify();
+#endif
+					ControlComponentsSubsystem->OnTrackedActorRegistered().Broadcast(Actor);
+				}
+
+				FRemoteControlComponentsUtils::RefreshTrackedProperties(Actor);
+			}
+		}
+	}
+#if WITH_EDITOR
+	AddTrackerTransaction.Reset();
+#endif
+}
+
+void FRemoteControlComponentsUtils::RemoveTrackerComponent(const TSet<TWeakObjectPtr<AActor>>& InActors)
+{
+#if WITH_EDITOR
+	FScopedTransaction Transaction(LOCTEXT("RemoveTrackerComponents", "Remove Tracker Component(s)"));
+#endif
+
+	for (const TWeakObjectPtr<AActor>& ActorWeak : InActors)
+	{
+		if (ActorWeak.IsValid())
+		{
+			if (AActor* Actor = ActorWeak.Get())
+			{
+				if (URemoteControlTrackerComponent* TrackerComponent = Actor->FindComponentByClass<URemoteControlTrackerComponent>())
+				{
+#if WITH_EDITOR
+					Actor->Modify();
+					TrackerComponent->Modify();
+
+					if (URemoteControlPreset* Preset = FRemoteControlComponentsUtils::GetCurrentPreset(Actor))
+					{
+						Preset->Modify();
+					}
+#endif
+					
+					TrackerComponent->UnexposeAllProperties();
+					Actor->RemoveInstanceComponent(TrackerComponent);
+					TrackerComponent->DestroyComponent();
+
+#if WITH_EDITOR
+					Actor->RerunConstructionScripts();
+#endif
+
+					if (URemoteControlComponentsSubsystem* ControlComponentsSubsystem = URemoteControlComponentsSubsystem::Get())
+					{
+#if WITH_EDITOR
+						ControlComponentsSubsystem->Modify();
+#endif
+						ControlComponentsSubsystem->OnTrackedActorUnregistered().Broadcast(Actor);
+					}
+				}
+			}
+		}
+	}
+}
+
+TWeakPtr<FRemoteControlProperty> FRemoteControlComponentsUtils::ExposeProperty(URemoteControlPreset* InPreset, UObject* InOwnerObject, const FRCFieldPathInfo& InPathInfo)
+{
+	if (!InPreset)
+	{
+		return nullptr;
+	}
+
+	FRemoteControlPresetExposeArgs Args;
+	Args.Label = FString();
+	Args.GroupId = FGuid();
+
+	const FGuid& PropertyGuid = FRemoteControlComponentsUtils::GetPropertyGuid(InOwnerObject, InPathInfo);
+	
+	if (!InPreset->IsExposed(PropertyGuid))
+	{
+		return InPreset->ExposeProperty(InOwnerObject, InPathInfo, Args);
+	}
+
+	return nullptr;
+}
+
+void FRemoteControlComponentsUtils::UnexposeProperty(URemoteControlPreset* InPreset, UObject* InOwnerObject, const FRCFieldPathInfo& InPathInfo)
+{
+	if (!InPreset)
+	{
+		return;
+	}
+
+	FRCFieldPathInfo FieldPathInfoLocal = InPathInfo;
+	FieldPathInfoLocal.Resolve(InOwnerObject);
+	
+	if (const FProperty* Property = FieldPathInfoLocal.GetResolvedData().Field)
+	{
+		const FGuid PropertyGuid = FRemoteControlComponentsUtils::GetPresetPropertyGuid(InPreset, InOwnerObject, Property);
+
+		if (InPreset->IsExposed(PropertyGuid))
+		{
+			InPreset->Unexpose(PropertyGuid);
+		}
+	}
+}
+
+TWeakPtr<FRemoteControlProperty> FRemoteControlComponentsUtils::GetExposedProperty(URemoteControlPreset* InPreset, UObject* InOwnerObject, const FRCFieldPathInfo& InPathInfo)
+{
+	if (!InPreset)
+	{
+		return nullptr;
+	}
+
+	FRCFieldPathInfo FieldPathInfoLocal = InPathInfo;
+	FieldPathInfoLocal.Resolve(InOwnerObject);
+	
+	if (const FProperty* Property = FieldPathInfoLocal.GetResolvedData().Field)
+	{
+		const FGuid PropertyGuid = FRemoteControlComponentsUtils::GetPresetPropertyGuid(InPreset, InOwnerObject, Property);
+
+		return InPreset->GetExposedEntity<FRemoteControlProperty>(PropertyGuid);
+	}
+
+	return nullptr;
+}
+
+FName FRemoteControlComponentsUtils::GetExposedPropertyId(URemoteControlPreset* InPreset, UObject* InOwnerObject, const FRCFieldPathInfo& InPathInfo)
+{
+	if (!InPreset)
+	{
+		return NAME_None;
+	}
+
+	const TWeakPtr<FRemoteControlProperty> RCPropertyWeak = FRemoteControlComponentsUtils::GetExposedProperty(InPreset, InOwnerObject, InPathInfo);		
+	if (const TSharedPtr<FRemoteControlProperty> RCProperty = RCPropertyWeak.Pin())
+	{
+		return RCProperty->PropertyId;
+	}
+
+	return NAME_None;
+}
+
+void FRemoteControlComponentsUtils::SetExposedPropertyId(URemoteControlPreset* InPreset, UObject* InOwnerObject, const FRCFieldPathInfo& InPathInfo, const FName& InPropertyId)
+{
+	if (!InPreset)
+	{
+		return;
+	}
+
+	const TWeakPtr<FRemoteControlProperty> RCPropertyWeak = FRemoteControlComponentsUtils::GetExposedProperty(InPreset, InOwnerObject, InPathInfo);		
+	if (const TSharedPtr<FRemoteControlProperty> RCProperty = RCPropertyWeak.Pin())
+	{
+		RCProperty->PropertyId = InPropertyId;
+		if (const TObjectPtr<URemoteControlPropertyIdRegistry>& PropertyRegistry = InPreset->GetPropertyIdRegistry())
+		{
+			PropertyRegistry->UpdateIdentifiedField(RCProperty.ToSharedRef());
+		}		
+	}
+}
+
+void FRemoteControlComponentsUtils::UnexposeAllProperties(AActor* InActor)
+{
+	if (!InActor)
+	{
+		return;
+	}
+
+	if (URemoteControlTrackerComponent* TrackerComponent = FRemoteControlComponentsUtils::GetTrackerComponent(InActor))
+	{
+#if WITH_EDITOR
+		FScopedTransaction Transaction(LOCTEXT("UnexposeAllProperties", "Unexpose all properties"));
+		InActor->Modify();
+
+		if (URemoteControlPreset* Preset = FRemoteControlComponentsUtils::GetCurrentPreset(InActor))
+		{
+			Preset->Modify();
+		}
+#endif
+		
+		TrackerComponent->UnexposeAllProperties();
+	}
+}
+
+void FRemoteControlComponentsUtils::RefreshTrackedProperties(AActor* InActor)
+{
+	if (!InActor)
+	{
+		return;
+	}
+	
+	constexpr bool bAddTrackerIfMissing = true;
+	URemoteControlTrackerComponent* TrackerComponent = GetTrackerComponent(InActor, bAddTrackerIfMissing);
+	
+	if (!TrackerComponent)
+	{
+		return;
+	}
+	
+	if (URemoteControlPreset* Preset = GetCurrentPreset(InActor))
+	{
+		const TArray<TWeakPtr<FRemoteControlProperty>>& RemoteControlExposedProperties = Preset->GetExposedEntities<FRemoteControlProperty>();
+		for (const TWeakPtr<FRemoteControlProperty>& RCPropertyWeak : RemoteControlExposedProperties)
+		{
+			if (const TSharedPtr<FRemoteControlProperty>& RCProperty = RCPropertyWeak.Pin())
+			{
+				UObject* BoundObject = RCProperty->GetBoundObject();
+				if (!BoundObject)
+				{
+					continue;
+				}
+				
+				if (InActor == BoundObject->GetTypedOuter<AActor>())
+				{
+					const FRCFieldPathInfo& FieldPathInfo = RCProperty->FieldPathInfo;
+					
+					if (!TrackerComponent->IsTrackingProperty(FieldPathInfo, BoundObject))
+					{
+						TrackerComponent->AddTrackedProperty(FieldPathInfo, BoundObject);
+					}				
+				}
+			}
+		}
+	}
+}
+
+FGuid FRemoteControlComponentsUtils::GetPropertyGuid(TConstArrayView<UObject*> InOuterObjects, const FProperty* InProperty)
+{
+	if (InOuterObjects.IsEmpty())
+	{
+		return FGuid();
+	}
+	
+	if (UObject* OwnerObject = InOuterObjects[0])
+	{
+		if (URemoteControlPreset* Preset = FRemoteControlComponentsUtils::GetCurrentPreset(OwnerObject))
+		{
+			return FRemoteControlComponentsUtils::GetPresetPropertyGuid(Preset, OwnerObject, InProperty);
+		}
+	}
+
+	return FGuid();
+}
+
+FGuid FRemoteControlComponentsUtils::GetPropertyGuid(UObject* InOwnerObject, const FRCFieldPathInfo& InFieldPathInfo)
+{
+	FRCFieldPathInfo FieldPathInfoLocal = InFieldPathInfo;
+	FieldPathInfoLocal.Resolve(InOwnerObject);
+	
+	if (const FProperty* Property = FieldPathInfoLocal.GetResolvedData().Field)
+	{
+		return GetPropertyGuid({InOwnerObject}, Property);
+	}
+
+	return FGuid();
+}
+
+#undef LOCTEXT_NAMESPACE
diff --git a/Engine/Plugins/Experimental/RemoteControlComponents/Source/RemoteControlComponents/Private/RemoteControlTrackerComponent.cpp b/Engine/Plugins/Experimental/RemoteControlComponents/Source/RemoteControlComponents/Private/RemoteControlTrackerComponent.cpp
new file mode 100644
--- /dev/null
+++ b/Engine/Plugins/Experimental/RemoteControlComponents/Source/RemoteControlComponents/Private/RemoteControlTrackerComponent.cpp
@@ -0,0 +1,263 @@
+// Copyright Epic Games, Inc. All Rights Reserved.
+
+#include "RemoteControlTrackerComponent.h"
+#include "Engine/World.h"
+#include "GameFramework/Actor.h"
+#include "Misc/TransactionObjectEvent.h"
+#include "RemoteControlTrackerProperty.h"
+#include "RemoteControlPreset.h"
+#include "Subsystems/RemoteControlComponentsSubsystem.h"
+
+void URemoteControlTrackerComponent::MarkPropertiesForRefresh()
+{
+	for (FRemoteControlTrackerProperty& TrackedProperty : TrackedProperties)
+	{
+		TrackedProperty.MarkUnexposed();
+	}
+}
+
+void URemoteControlTrackerComponent::RefreshExposedProperties()
+{
+	MarkPropertiesForRefresh();
+	ExposeAllProperties();
+}
+
+void URemoteControlTrackerComponent::RegisterTrackedActor() const
+{	
+	if (AActor* OuterActor = GetTrackedActor())
+	{
+		if (URemoteControlComponentsSubsystem* RemoteControlComponentsSubsystem = URemoteControlComponentsSubsystem::Get())
+		{
+			RemoteControlComponentsSubsystem->RegisterTrackedActor(OuterActor);
+		}
+	}
+}
+
+void URemoteControlTrackerComponent::UnregisterTrackedActor() const
+{
+	if (AActor* OuterActor = GetTrackedActor())
+	{
+		if (URemoteControlComponentsSubsystem* RemoteControlComponentsSubsystem = URemoteControlComponentsSubsystem::Get())
+		{
+			RemoteControlComponentsSubsystem->UnregisterTrackedActor(OuterActor);
+		}
+	}
+}
+
+void URemoteControlTrackerComponent::OnComponentCreated()
+{
+	Super::OnComponentCreated();
+
+	if (AActor* TrackedActor = GetTrackedActor())
+	{
+		if (URemoteControlComponentsSubsystem* RemoteControlComponentsSubsystem = URemoteControlComponentsSubsystem::Get())
+		{
+			RemoteControlComponentsSubsystem->RegisterTrackedActor(TrackedActor);
+		}
+	}
+}
+
+void URemoteControlTrackerComponent::OnComponentDestroyed(bool bInDestroyingHierarchy)
+{
+	Super::OnComponentDestroyed(bInDestroyingHierarchy);
+
+	UnexposeAllProperties();
+	
+	if (AActor* TrackedActor = GetTrackedActor())
+	{
+		if (URemoteControlComponentsSubsystem* RemoteControlComponentsSubsystem = URemoteControlComponentsSubsystem::Get())
+		{
+			RemoteControlComponentsSubsystem->UnregisterTrackedActor(TrackedActor);
+		}
+	}
+}
+
+void URemoteControlTrackerComponent::RefreshTracker()
+{
+	RegisterTrackedActor();
+	RefreshExposedProperties();
+}
+
+void URemoteControlTrackerComponent::PostInitProperties()
+{
+	Super::PostInitProperties();
+
+	RegisterTrackedActor();
+	RegisterPropertyIdChangeDelegate();
+}
+
+void URemoteControlTrackerComponent::OnTrackerDuplicated()
+{
+	RefreshTracker();
+	WriteAllPropertyIdsToPreset();
+	RegisterPropertyIdChangeDelegate();
+}
+
+void URemoteControlTrackerComponent::PostDuplicate(bool bInDuplicateForPIE)
+{
+	Super::PostDuplicate(bInDuplicateForPIE);
+
+	OnTrackerDuplicated();
+}
+
+void URemoteControlTrackerComponent::PostEditImport()
+{
+	Super::PostEditImport();
+
+	OnTrackerDuplicated();
+}
+
+void URemoteControlTrackerComponent::PostLoad()
+{
+	Super::PostLoad();
+	RegisterPropertyIdChangeDelegate();
+}
+
+#if WITH_EDITOR
+void URemoteControlTrackerComponent::PostTransacted(const FTransactionObjectEvent& InTransactionEvent)
+{
+	Super::PostTransacted(InTransactionEvent);
+
+	if (InTransactionEvent.GetEventType() == ETransactionObjectEventType::UndoRedo)
+	{
+		RefreshTracker();
+		
+		// In case the component was destroyed and created again with an Undo, let's try fixing potentially broken bindings
+		if (URemoteControlPreset* Preset = GetCurrentPreset())
+		{
+			Preset->RebindUnboundEntities();
+		}
+	}
+}
+#endif
+
+int32 URemoteControlTrackerComponent::GetTrackedPropertyIndex(const FRCFieldPathInfo& InFieldPathInfo, UObject* InOwnerObject) const
+{
+	return TrackedProperties.IndexOfByPredicate([InFieldPathInfo, InOwnerObject](const FRemoteControlTrackerProperty& TrackerProperty)
+	{
+		return TrackerProperty.MatchesParameters(InFieldPathInfo, InOwnerObject);
+	});
+}
+
+void URemoteControlTrackerComponent::AddTrackedProperty(const FRCFieldPathInfo& InFieldPathInfo, UObject* InOwnerObject)
+{
+	if (!IsTrackingProperty(InFieldPathInfo, InOwnerObject))
+	{
+		TrackedProperties.Add(FRemoteControlTrackerProperty(InFieldPathInfo, InOwnerObject));
+	}
+}
+
+void URemoteControlTrackerComponent::RemoveTrackedProperty(const FRCFieldPathInfo& InFieldPathInfo, UObject* InOwnerObject)
+{
+	const int32 Index = GetTrackedPropertyIndex(InFieldPathInfo, InOwnerObject);
+	if (Index != INDEX_NONE)
+	{
+		TrackedProperties.RemoveAt(Index);
+	}
+}
+
+URemoteControlPreset* URemoteControlTrackerComponent::GetCurrentPreset() const
+{
+	if (const URemoteControlComponentsSubsystem* RemoteControlPresetComponentsSubsystem = URemoteControlComponentsSubsystem::Get())
+	{
+		if (const UWorld* World = GetTypedOuter<UWorld>())
+		{
+			if (URemoteControlPreset* Preset = RemoteControlPresetComponentsSubsystem->GetRegisteredPreset(World))
+			{
+				return Preset;
+			}
+		}
+	}
+
+	return nullptr;
+}
+
+void URemoteControlTrackerComponent::ExposeAllProperties()
+{
+	if (URemoteControlPreset* Preset = GetCurrentPreset())
+	{
+		for (FRemoteControlTrackerProperty& TrackedProperty : TrackedProperties)
+		{
+			TrackedProperty.Expose(Preset);
+		}
+	}
+}
+
+bool URemoteControlTrackerComponent::HasTrackedProperties() const
+{
+	return !TrackedProperties.IsEmpty();
+}
+
+bool URemoteControlTrackerComponent::IsTrackingProperty(const FRCFieldPathInfo& InFieldPathInfo, UObject* InOwnerObject) const
+{
+	const int32 Index = GetTrackedPropertyIndex(InFieldPathInfo, InOwnerObject);
+	return Index != INDEX_NONE;
+}
+
+void URemoteControlTrackerComponent::UnexposeAllProperties()
+{
+	TArray<FRemoteControlTrackerProperty> PropertiesToRemove = MoveTemp(TrackedProperties);
+
+	for (FRemoteControlTrackerProperty& TrackedProperty : PropertiesToRemove)
+	{
+		TrackedProperty.Unexpose();
+	}
+	
+	UnregisterPropertyIdChangeDelegate();
+}
+
+void URemoteControlTrackerComponent::RefreshAllPropertyIds()
+{
+	for (FRemoteControlTrackerProperty& TrackedProperty : TrackedProperties)
+	{
+		TrackedProperty.ReadPropertyIdFromPreset();
+	}
+}
+
+void URemoteControlTrackerComponent::WriteAllPropertyIdsToPreset()
+{
+	UnregisterPropertyIdChangeDelegate();
+	for (FRemoteControlTrackerProperty& TrackedProperty : TrackedProperties)
+	{
+		TrackedProperty.WritePropertyIdToPreset();
+	}
+	RegisterPropertyIdChangeDelegate();
+}
+
+AActor* URemoteControlTrackerComponent::GetTrackedActor() const
+{
+	if (AActor* OuterActor = GetOwner())
+	{
+		return OuterActor;
+	}
+
+	return nullptr;
+}
+
+void URemoteControlTrackerComponent::RegisterPropertyIdChangeDelegate()
+{
+	if (const URemoteControlPreset* Preset = GetCurrentPreset())
+	{
+		if (const TObjectPtr<URemoteControlPropertyIdRegistry>& PropertyRegistry = Preset->GetPropertyIdRegistry())
+		{
+			PropertyRegistry->OnPropertyIdUpdated().RemoveAll(this);
+			PropertyRegistry->OnPropertyIdUpdated().AddUObject(this, &URemoteControlTrackerComponent::OnPropertyIdUpdated);
+		}
+	}
+}
+
+void URemoteControlTrackerComponent::UnregisterPropertyIdChangeDelegate() const
+{
+	if (const URemoteControlPreset* Preset = GetCurrentPreset())
+	{
+		if (const TObjectPtr<URemoteControlPropertyIdRegistry>& PropertyRegistry = Preset->GetPropertyIdRegistry())
+		{
+			PropertyRegistry->OnPropertyIdUpdated().RemoveAll(this);
+		}
+	}
+}
+
+void URemoteControlTrackerComponent::OnPropertyIdUpdated()
+{
+	RefreshAllPropertyIds();
+}
diff --git a/Engine/Plugins/Experimental/RemoteControlComponents/Source/RemoteControlComponents/Private/RemoteControlTrackerProperty.cpp b/Engine/Plugins/Experimental/RemoteControlComponents/Source/RemoteControlComponents/Private/RemoteControlTrackerProperty.cpp
new file mode 100644
--- /dev/null
+++ b/Engine/Plugins/Experimental/RemoteControlComponents/Source/RemoteControlComponents/Private/RemoteControlTrackerProperty.cpp
@@ -0,0 +1,115 @@
+// Copyright Epic Games, Inc. All Rights Reserved.
+
+#include "RemoteControlTrackerProperty.h"
+#include "IRemoteControlModule.h"
+#include "RemoteControlComponentsUtils.h"
+#include "RemoteControlPreset.h"
+
+FRemoteControlTrackerProperty::FRemoteControlTrackerProperty(const FRCFieldPathInfo& InFieldPathInfo, const TObjectPtr<UObject>& InOwnerObject, bool bInIsExposed)
+	: FieldPathInfo(InFieldPathInfo)
+	, OwnerObject(InOwnerObject)
+	, bIsExposed(bInIsExposed)
+{
+	Resolve();
+	ReadPropertyIdFromPreset();
+}
+
+bool FRemoteControlTrackerProperty::MatchesParameters(const FRCFieldPathInfo& InFieldPathInfo, const TObjectPtr<UObject>& InOwnerObject) const
+{
+	return FieldPathInfo == InFieldPathInfo
+		&& OwnerObject == InOwnerObject;
+}
+
+URemoteControlPreset* FRemoteControlTrackerProperty::GetPreset() const
+{
+	return FRemoteControlComponentsUtils::GetCurrentPreset(OwnerObject.Get());
+}
+
+bool FRemoteControlTrackerProperty::IsExposedTo(const URemoteControlPreset* InPreset) const
+{
+	if (const URemoteControlPreset* CurrentPreset = GetPreset())
+	{
+		if (CurrentPreset == InPreset)
+		{
+			return IsExposed();
+		}
+	}
+
+	return false;
+}
+
+void FRemoteControlTrackerProperty::Resolve() const
+{
+	if (OwnerObject.IsValid())
+	{
+		const_cast<FRemoteControlTrackerProperty*>(this)->FieldPathInfo.Resolve(OwnerObject.Get());
+	}
+}
+
+void FRemoteControlTrackerProperty::MarkUnexposed()
+{
+	bIsExposed = false;
+}
+
+void FRemoteControlTrackerProperty::Expose(URemoteControlPreset* InRemoteControlPreset)
+{
+	if (IsExposedTo(InRemoteControlPreset))
+	{
+		// Property is already exposed to specified preset, just return
+		return;
+	}
+	
+	if (!InRemoteControlPreset || !OwnerObject.IsValid())
+	{
+		return;
+	}
+	
+	FRemoteControlComponentsUtils::ExposeProperty(InRemoteControlPreset, OwnerObject.Get(), FieldPathInfo);
+	CurrentPresetWeak = MakeWeakObjectPtr<URemoteControlPreset>(InRemoteControlPreset);
+	bIsExposed = true;	
+}
+
+void FRemoteControlTrackerProperty::Unexpose()
+{
+	if (!IsExposed())
+	{
+		return;
+	}
+	
+	if (!OwnerObject.IsValid())
+	{
+		return;
+	}
+
+	if (URemoteControlPreset* Preset = GetPreset())
+	{
+		FRemoteControlComponentsUtils::UnexposeProperty(Preset, OwnerObject.Get(), FieldPathInfo);
+	}
+
+	bIsExposed = false;
+}
+
+void FRemoteControlTrackerProperty::ReadPropertyIdFromPreset()
+{
+	if (URemoteControlPreset* Preset = GetPreset())
+	{
+		const FName& Id = FRemoteControlComponentsUtils::GetExposedPropertyId(Preset, OwnerObject.Get(), FieldPathInfo);
+		if (PropertyId != Id)
+		{
+			PropertyId = Id;
+		}
+	}
+}
+
+void FRemoteControlTrackerProperty::WritePropertyIdToPreset() const
+{
+	if (!IsExposed())
+	{
+		return;
+	}
+
+	if (URemoteControlPreset* Preset = GetPreset())
+	{
+		FRemoteControlComponentsUtils::SetExposedPropertyId(Preset, OwnerObject.Get(), FieldPathInfo, PropertyId);
+	}
+}
diff --git a/Engine/Plugins/Experimental/RemoteControlComponents/Source/RemoteControlComponents/Private/Subsystems/RemoteControlComponentsSubsystem.cpp b/Engine/Plugins/Experimental/RemoteControlComponents/Source/RemoteControlComponents/Private/Subsystems/RemoteControlComponentsSubsystem.cpp
new file mode 100644
--- /dev/null
+++ b/Engine/Plugins/Experimental/RemoteControlComponents/Source/RemoteControlComponents/Private/Subsystems/RemoteControlComponentsSubsystem.cpp
@@ -0,0 +1,378 @@
+// Copyright Epic Games, Inc. All Rights Reserved.
+
+#include "Subsystems/RemoteControlComponentsSubsystem.h"
+#include "Engine/Engine.h"
+#include "IRemoteControlModule.h"
+#include "RemoteControlComponentsUtils.h"
+#include "RemoteControlField.h"
+#include "RemoteControlPreset.h"
+#include "RemoteControlTrackerComponent.h"
+
+#if WITH_EDITOR
+#include "ScopedTransaction.h"
+#endif
+
+#define LOCTEXT_NAMESPACE "RemoteControlComponentsSubsystem"
+
+URemoteControlComponentsSubsystem* URemoteControlComponentsSubsystem::Get()
+{
+	if (GEngine)
+	{
+		if (URemoteControlComponentsSubsystem* RemoteControlComponentsSubsystem = GEngine->GetEngineSubsystem<URemoteControlComponentsSubsystem>())
+		{
+			return RemoteControlComponentsSubsystem;
+		}
+	}
+	return nullptr;
+}
+
+void URemoteControlComponentsSubsystem::Initialize(FSubsystemCollectionBase& InCollection)
+{
+	Super::Initialize(InCollection);
+	
+#if WITH_EDITOR
+	FWorldDelegates::OnCurrentLevelChanged.AddUObject(this, &URemoteControlComponentsSubsystem::OnLevelChanged);
+	FWorldDelegates::OnWorldCleanup.AddUObject(this, &URemoteControlComponentsSubsystem::OnWorldCleanup);
+	
+	if (GEngine)
+	{
+		GEngine->OnLevelActorDeleted().AddUObject(this, &URemoteControlComponentsSubsystem::OnLevelActorDeleted);
+	}
+#endif
+}
+
+void URemoteControlComponentsSubsystem::Deinitialize()
+{
+	Super::Deinitialize();
+
+#if WITH_EDITOR
+	FWorldDelegates::OnCurrentLevelChanged.RemoveAll(this);
+	FWorldDelegates::OnWorldCleanup.RemoveAll(this);
+	
+	if (GEngine)
+	{
+		GEngine->OnLevelActorDeleted().RemoveAll(this);
+	}
+#endif
+}
+
+void URemoteControlComponentsSubsystem::OnLevelChanged(ULevel* InNewLevel, ULevel* InOldLevel, UWorld* InWorld)
+{
+	ClearData();
+}
+
+void URemoteControlComponentsSubsystem::OnWorldCleanup(UWorld* InWorld, bool bInSessionEnded, bool bInCleanupResources)
+{
+	ClearData();
+}
+
+void URemoteControlComponentsSubsystem::ClearData()
+{
+	TrackedActorsContextMap.Empty();
+	ActorsPendingRegister.Empty();
+}
+
+void URemoteControlComponentsSubsystem::CacheActorsToBeTracked(const TSet<TWeakObjectPtr<AActor>>& InActors)
+{
+	if (InActors.IsEmpty())
+	{
+		return;
+	}
+	
+	for (TWeakObjectPtr<AActor> Actor: InActors)
+	{
+		if (Actor.IsValid())
+		{
+			ActorsPendingRegister.AddUnique(Actor.Get());
+		}
+	}
+}
+
+void URemoteControlComponentsSubsystem::RegisterTrackedActor(AActor* InActor)
+{
+	if (!InActor || IsActorTracked(InActor))
+	{
+		return;
+	}
+
+	bool bRegistered = false;
+	if (const UWorld* ActorWorld = InActor->GetWorld())
+	{
+		if (const TSharedPtr<FRemoteControlComponentsContext>& TrackingContext = GetTrackedActorContext(ActorWorld))
+		{
+			if (TrackingContext->IsValid())
+			{
+				bRegistered = TrackingContext->RegisterActor(InActor);
+			}
+		}
+	}
+
+	if (bRegistered)
+	{
+		OnTrackedActorRegisteredDelegate.Broadcast(InActor);
+	}
+	else
+	{
+		// It was not possible to register the actor at this time.
+		// We cache it so we can try to register it again when a new Preset is registered
+		CacheActorsToBeTracked({InActor});
+	}
+}
+
+void URemoteControlComponentsSubsystem::UnregisterTrackedActor(AActor* InActor)
+{
+	if (InActor)
+	{
+		if (const UWorld* ActorWorld = InActor->GetWorld())
+		{
+			if (const TSharedPtr<FRemoteControlComponentsContext>& TrackingContext = GetTrackedActorContext(ActorWorld))
+			{
+				if (TrackingContext->UnregisterActor(InActor))
+				{
+					OnTrackedActorUnregisteredDelegate.Broadcast(InActor);
+				}
+			}
+		}
+
+		ActorsPendingRegister.Remove(InActor);
+	}
+}
+
+void URemoteControlComponentsSubsystem::RegisterPendingActors(const UWorld* InWorld)
+{
+	TArray<TWeakObjectPtr<AActor>> RegisteredActors;
+	for (const TWeakObjectPtr<AActor>& PendingActor : ActorsPendingRegister)
+	{
+		if (!PendingActor.IsValid())
+		{
+			continue;
+		}
+			
+		if (PendingActor->GetWorld() == InWorld)
+		{
+			RegisterTrackedActor(PendingActor.Get());
+			RegisteredActors.AddUnique(PendingActor);
+		}
+	}
+
+	for (const TWeakObjectPtr<AActor>& PendingActor : RegisteredActors)
+	{
+		ActorsPendingRegister.Remove(PendingActor);
+	}
+}
+
+void URemoteControlComponentsSubsystem::RegisterPreset(URemoteControlPreset* InRemoteControlPreset)
+{
+	if (!InRemoteControlPreset)
+	{
+		return;
+	}
+	
+	if (UWorld* PresetWorld = InRemoteControlPreset->GetTypedOuter<UWorld>())
+	{
+		const FObjectKey WorldKey(PresetWorld);
+		if (!TrackedActorsContextMap.Contains(WorldKey))
+		{
+			TrackedActorsContextMap.Add(WorldKey, MakeShared<FRemoteControlComponentsContext>(PresetWorld, InRemoteControlPreset));
+		}
+
+		// A new Preset has been registered, let's try re-registering pending actors
+		RegisterPendingActors(PresetWorld);
+
+		InRemoteControlPreset->OnEntityExposed().AddUObject(this, &URemoteControlComponentsSubsystem::OnEntityExposed);
+		InRemoteControlPreset->OnEntityUnexposed().AddUObject(this, &URemoteControlComponentsSubsystem::OnEntityUnexposed);
+	}
+}
+
+void URemoteControlComponentsSubsystem::UnregisterPreset(URemoteControlPreset* InRemoteControlPreset)
+{
+	if (!InRemoteControlPreset)
+	{
+		return;
+	}
+	
+	if (const UWorld* PresetWorld = InRemoteControlPreset->GetTypedOuter<UWorld>())
+	{
+		const FObjectKey WorldKey(PresetWorld);
+		if (TrackedActorsContextMap.Contains(WorldKey))
+		{
+			TrackedActorsContextMap.Remove(WorldKey);
+		}
+	}
+
+	InRemoteControlPreset->OnEntityExposed().RemoveAll(this);
+	InRemoteControlPreset->OnEntityUnexposed().RemoveAll(this);
+}
+
+URemoteControlPreset* URemoteControlComponentsSubsystem::GetRegisteredPreset(const UWorld* InWorld) const
+{
+	if (InWorld)
+	{
+		if (const TSharedPtr<FRemoteControlComponentsContext>& TrackingContext = GetTrackedActorContext(InWorld))
+		{
+			if (URemoteControlPreset* RegisteredPreset = TrackingContext->GetPreset())
+			{
+				return RegisteredPreset;
+			}
+		}
+	}
+	
+	return nullptr;
+}
+
+URemoteControlPreset* URemoteControlComponentsSubsystem::GetRegisteredPreset(const UObject* InObject) const
+{
+	if (!InObject)
+	{
+		return nullptr;
+	}
+	
+	if (const UWorld* World = InObject->GetTypedOuter<UWorld>())
+	{
+		return GetRegisteredPreset(World);
+	}
+
+	return nullptr;
+}
+
+bool URemoteControlComponentsSubsystem::IsActorTracked(const AActor* InActor) const
+{
+	if (!InActor)
+	{
+		return false;
+	}
+
+	if (const UWorld* ActorWorld = InActor->GetWorld())
+	{
+		if (const TSharedPtr<FRemoteControlComponentsContext>& TrackingContext = GetTrackedActorContext(ActorWorld))
+		{
+			return TrackingContext->IsActorRegistered(InActor);
+		}
+	}
+
+	return false;
+}
+
+bool URemoteControlComponentsSubsystem::IsPresetRegistered(const URemoteControlPreset* InPreset) const
+{
+	if (!InPreset)
+	{
+		return false;
+	}
+
+	if (const UWorld* PresetWorld = InPreset->GetTypedOuter<UWorld>())
+	{	
+		if (const URemoteControlPreset* RegisteredPreset = GetRegisteredPreset(PresetWorld))
+		{
+			return RegisteredPreset == InPreset;
+		}
+	}
+
+	return false;
+}
+
+void URemoteControlComponentsSubsystem::OnEntityExposed(URemoteControlPreset* InRemoteControlPreset, const FGuid& InGuid)
+{
+	if (InRemoteControlPreset && InGuid.IsValid())
+	{
+		if (const TSharedPtr<FRemoteControlProperty>& RCProperty = InRemoteControlPreset->GetExposedEntity<FRemoteControlProperty>(InGuid).Pin())
+		{
+			if (RCProperty.IsValid())
+			{
+				constexpr bool bCleanDuplicates = true;
+				const FRCFieldPathInfo PathInfo(RCProperty->FieldPathInfo.ToString(), bCleanDuplicates);
+
+				if (UObject* OwnerObject = RCProperty->GetBoundObject())
+				{
+					if (AActor* OwnerActor = OwnerObject->GetTypedOuter<AActor>())
+					{
+						constexpr bool bAddTrackerIfMissing = true;
+						if (URemoteControlTrackerComponent* TrackerComponent = FRemoteControlComponentsUtils::GetTrackerComponent(OwnerActor, bAddTrackerIfMissing))
+						{
+							TrackerComponent->AddTrackedProperty(PathInfo, OwnerObject);
+						}
+					}
+				}
+			}
+		}
+	}
+}
+
+void URemoteControlComponentsSubsystem::OnEntityUnexposed(URemoteControlPreset* InRemoteControlPreset, const FGuid& InGuid)
+{
+	if (!InRemoteControlPreset || !InGuid.IsValid())
+	{
+		return;
+	}
+
+	const TSharedPtr<FRemoteControlProperty>& RCProperty = InRemoteControlPreset->GetExposedEntity<FRemoteControlProperty>(InGuid).Pin();
+	if (!RCProperty.IsValid())
+	{
+		return;
+	}
+
+	UObject* OwnerObject = RCProperty->GetBoundObject();
+	if (!OwnerObject)
+	{
+		return;
+	}
+
+	if (AActor* OwnerActor = OwnerObject->GetTypedOuter<AActor>())
+	{
+		// Look for a tracker component, and remove the property from the tracked properties set
+		if (URemoteControlTrackerComponent* TrackerComponent = FRemoteControlComponentsUtils::GetTrackerComponent(OwnerActor))
+		{			
+			if (IsRegisteredTrackedActor(OwnerActor))
+			{
+				constexpr bool bCleanDuplicates = true;
+				const FRCFieldPathInfo PathInfo(RCProperty->FieldPathInfo.ToString(), bCleanDuplicates);
+				TrackerComponent->RemoveTrackedProperty(PathInfo, OwnerObject);
+			}
+		}
+	}
+}
+
+void URemoteControlComponentsSubsystem::OnLevelActorDeleted(AActor* InActor)
+{
+	if (!InActor)
+	{
+		return;
+	}
+	
+	UnregisterTrackedActor(InActor);
+	FRemoteControlComponentsUtils::UnexposeAllProperties(InActor);
+}
+
+bool URemoteControlComponentsSubsystem::IsRegisteredActor(const AActor* InActor) const
+{
+	return IsRegisteredTrackedActor(InActor);
+}
+
+bool URemoteControlComponentsSubsystem::IsRegisteredTrackedActor(const AActor* InActor) const
+{
+	if (const UWorld* ActorWorld = InActor->GetWorld())
+	{
+		if (const TSharedPtr<FRemoteControlComponentsContext>& TrackingContext = GetTrackedActorContext(ActorWorld))
+		{
+			return TrackingContext->IsActorRegistered(InActor);
+		}
+	}
+	
+	return false;
+}
+
+TSharedPtr<FRemoteControlComponentsContext> URemoteControlComponentsSubsystem::GetTrackedActorContext(const UWorld* InWorld) const
+{
+	if (InWorld)
+	{
+		const FObjectKey WorldKey(InWorld);
+		if (const TSharedPtr<FRemoteControlComponentsContext>* TrackingContextPtr = TrackedActorsContextMap.Find(WorldKey))
+		{
+			return *TrackingContextPtr;
+		}
+	}
+
+	return {};
+}
+
+#undef LOCTEXT_NAMESPACE
diff --git a/Engine/Plugins/Experimental/RemoteControlComponents/Source/RemoteControlComponents/Public/RemoteControlComponentsContext.h b/Engine/Plugins/Experimental/RemoteControlComponents/Source/RemoteControlComponents/Public/RemoteControlComponentsContext.h
new file mode 100644
--- /dev/null
+++ b/Engine/Plugins/Experimental/RemoteControlComponents/Source/RemoteControlComponents/Public/RemoteControlComponentsContext.h
@@ -0,0 +1,62 @@
+﻿// Copyright Epic Games, Inc. All Rights Reserved.
+
+#pragma once
+
+#include "Containers/Array.h"
+#include "UObject/WeakObjectPtr.h"
+
+class AActor;
+class URemoteControlPreset;
+class UWorld;
+
+/**
+ * Remote Control Components Tracking Context: represents a World with a set of Actors which have properties exposed to a specific Preset.
+ * It currently works under the assumption that there's a single Preset per World, and that this info is runtime only. Both aspects might change in the future.
+ *
+ * Context is represented by:
+ *	- a World
+ *	- a Remote Control Preset
+ *	- Tracked Actors with properties exposed to the Preset
+ */
+struct FRemoteControlComponentsContext
+{
+public:
+	/** Create a Remote Control Context by passing a World and a Preset */
+	FRemoteControlComponentsContext(UWorld* InWorld, URemoteControlPreset* InPreset);
+
+	/** Adds an Actor to the list of tracked actors */
+	bool RegisterActor(AActor* InActor);
+
+	/** Removes an Actor from the list of tracked actors */
+	bool UnregisterActor(AActor* InActor);
+
+	/** Checks if the specified Actor is tracked or not */
+	bool IsActorRegistered(const AActor* InActor) const;
+
+	/** An Actor can be registered if it has a Remote Control Tracker Component and its World and the World of this context coincide */
+	bool CanRegisterActor(const AActor* InActor) const;
+
+	/** Context is valid if both World and Preset are valid */
+	bool IsValid() const;
+
+	/** True if the context currently has registered tracked actors */
+	bool HasTrackedActors() const;
+
+	/** Get the array of currently tracked actors */
+	TArray<TWeakObjectPtr<AActor>> GetTrackedActors() const { return TrackedActors; }
+
+	/** Get the World referenced by this context */
+	UWorld* GetWorld() const;
+
+	/** Get the Remote Control Preset referenced by this context */
+	URemoteControlPreset* GetPreset() const;
+
+private:
+	TWeakObjectPtr<UWorld> WorldWeak;
+	
+	TWeakObjectPtr<URemoteControlPreset> PresetWeak;
+	
+	TArray<TWeakObjectPtr<AActor>> TrackedActors;
+
+	bool bIsValid = false;
+};
diff --git a/Engine/Plugins/Experimental/RemoteControlComponents/Source/RemoteControlComponents/Public/RemoteControlComponentsUtils.h b/Engine/Plugins/Experimental/RemoteControlComponents/Source/RemoteControlComponents/Public/RemoteControlComponentsUtils.h
new file mode 100644
--- /dev/null
+++ b/Engine/Plugins/Experimental/RemoteControlComponents/Source/RemoteControlComponents/Public/RemoteControlComponentsUtils.h
@@ -0,0 +1,98 @@
+﻿// Copyright Epic Games, Inc. All Rights Reserved.
+
+#pragma once
+
+#include "HAL/Platform.h"
+#include "Templates/SharedPointer.h"
+#include "UObject/WeakObjectPtrTemplates.h"
+
+class AActor;
+class URemoteControlTrackerComponent;
+class URemoteControlPreset;
+
+struct FRCFieldPathInfo;
+struct FRemoteControlProperty;
+
+class FRemoteControlComponentsUtils
+{
+	friend class URemoteControlComponentsSubsystem;
+	friend struct FRemoteControlTrackerProperty;
+public:
+	/** Add a RemoteControlTrackerComponent to the specified Actors */
+	REMOTECONTROLCOMPONENTS_API static void AddTrackerComponent(const TSet<TWeakObjectPtr<AActor>>& InActors, bool bInShouldTransact = true);
+
+	/** Removes the RemoteControlTrackerComponent from the specified Actors */
+	REMOTECONTROLCOMPONENTS_API static void RemoveTrackerComponent(const TSet<TWeakObjectPtr<AActor>>& InActors);
+
+	/**
+	 * Gets the tracker component for the specified Actor. Tracker can be optionally added if missing.
+	 * @param InActor: the Actor for which we want to retrieve the Tracker Component. Object 
+	 * @param bInAddTrackerIfMissing: if true, function will add a tracker component when one is missing
+	 * @return The found or created tracker.
+	 */
+	REMOTECONTROLCOMPONENTS_API static URemoteControlTrackerComponent* GetTrackerComponent(AActor* InActor, bool bInAddTrackerIfMissing = false);
+
+	/** Gets the Remote Control Preset currently handling the specified Object */
+	REMOTECONTROLCOMPONENTS_API static URemoteControlPreset* GetCurrentPreset(const UObject* InObject);
+
+	/** Gets the Remote Control Preset currently handling the specified World */
+	REMOTECONTROLCOMPONENTS_API static URemoteControlPreset* GetCurrentPreset(const UWorld* InWorld);
+
+	/**
+	 * Unexpose all Tracked Properties of the specified Actor.
+	 * It will also stop tracking all unexposed properties.
+	 */
+	REMOTECONTROLCOMPONENTS_API static void UnexposeAllProperties(AActor* InActor);
+
+	/** Call this function on any Actor to ensure its exposed RC properties are actually tracked */
+	REMOTECONTROLCOMPONENTS_API static void RefreshTrackedProperties(AActor* InActor);
+	
+private:
+	/**
+	 * Expose a Property to the specified Preset
+	 * @param InPreset the Preset where the property should be exposed
+	 * @param InOwnerObject the object owning the property
+	 * @param InPathInfo Field Path Info for the desire property
+	 * @return The exposed property
+	 */
+	static TWeakPtr<FRemoteControlProperty> ExposeProperty(URemoteControlPreset* InPreset, UObject* InOwnerObject, const FRCFieldPathInfo& InPathInfo);
+
+	/**
+	 * Unexpose a Property in the specified Preset
+	 * @param InPreset the Preset where the property is currently exposed
+	 * @param InOwnerObject the object owning the property
+	 * @param InPathInfo Field Path Info for the desire property
+	 */
+	static void UnexposeProperty(URemoteControlPreset* InPreset, UObject* InOwnerObject, const FRCFieldPathInfo& InPathInfo);
+
+	static TWeakPtr<FRemoteControlProperty> GetExposedProperty(URemoteControlPreset* InPreset, UObject* InOwnerObject, const FRCFieldPathInfo& InPathInfo);
+
+	/**
+	 * Returns the field identifier for the exposed property described by the provided arguments.
+	 * @param InPreset the Preset where the property is currently exposed
+	 * @param InOwnerObject the object owning the property
+	 * @param InPathInfo Field Path Info for the desire property
+	 */
+	static FName GetExposedPropertyId(URemoteControlPreset* InPreset, UObject* InOwnerObject, const FRCFieldPathInfo& InPathInfo);
+
+	/**
+	 * Sets the field identifier for the exposed property described by the provided arguments.
+	 * @param InPreset the Preset where the property is currently exposed
+	 * @param InOwnerObject the object owning the property
+	 * @param InPathInfo Field Path Info for the desire property
+	 * @param InPropertyId The new Property Id value to be set
+	 */
+	static void SetExposedPropertyId(URemoteControlPreset* InPreset, UObject* InOwnerObject, const FRCFieldPathInfo& InPathInfo, const FName& InPropertyId);
+	
+	static FGuid GetPropertyGuid(TConstArrayView<UObject*> InOuterObjects, const FProperty* InProperty);
+	static FGuid GetPropertyGuid(UObject* InOwnerObject, const FRCFieldPathInfo& InFieldPathInfo);
+
+	/**
+	 * Returns the ID of a property, if exposed in a Preset for the specified Owner Object
+	 * @param InPreset The Preset containing the exposed Property
+	 * @param InOwnerObject The Object owning the bound property
+	 * @param InProperty The property to look for
+	 * @return the ID of the property
+	 */
+	static FGuid GetPresetPropertyGuid(URemoteControlPreset* InPreset, UObject* InOwnerObject, const FProperty* InProperty);
+};
diff --git a/Engine/Plugins/Experimental/RemoteControlComponents/Source/RemoteControlComponents/Public/RemoteControlTrackerComponent.h b/Engine/Plugins/Experimental/RemoteControlComponents/Source/RemoteControlComponents/Public/RemoteControlTrackerComponent.h
new file mode 100644
--- /dev/null
+++ b/Engine/Plugins/Experimental/RemoteControlComponents/Source/RemoteControlComponents/Public/RemoteControlTrackerComponent.h
@@ -0,0 +1,96 @@
+// Copyright Epic Games, Inc. All Rights Reserved.
+
+#pragma once
+
+#include "Components/SceneComponent.h"
+#include "RemoteControlTrackerProperty.h"
+#include "RemoteControlTrackerComponent.generated.h"
+
+struct FRCFieldPathInfo;
+class URemoteControlPreset;
+
+/**
+ * A component keeping track of properties currently exposed to Remote Control.
+ * Supports auto-exposing of properties upon duplication.
+ * e.g. duplicating an Actor with this component will result in the same properties exposed by Source Actor
+ * to be automatically exposed by Duplicate Actor.
+ */
+UCLASS()
+class REMOTECONTROLCOMPONENTS_API URemoteControlTrackerComponent : public UActorComponent
+{
+	GENERATED_BODY()
+
+public:
+	/** Add the specified property to the list of properties handled by this Tracker Component */
+	void AddTrackedProperty(const FRCFieldPathInfo& InFieldPathInfo, UObject* InOwnerObject);
+	
+	/** Remove the specified property from the list of properties handled by this Tracker Component */
+	void RemoveTrackedProperty(const FRCFieldPathInfo& InFieldPathInfo, UObject* InOwnerObject);
+
+	/** Expose all Tracked Properties handled by this component to the current Remote Control Preset */
+	void ExposeAllProperties();
+	
+	/**
+	 * Unexpose all Tracked Properties handled by this component from the current Remote Control Preset.
+	 * It will also stop tracking all unexposed properties.
+	 */
+	void UnexposeAllProperties();
+
+	/**
+	 * Retrieves the Property Id for all the exposed tracked properties from the current Preset
+	 * (Tracker --> Preset)
+	 */
+	void RefreshAllPropertyIds();
+
+	/**
+	 * Writes the Property Id saved in the Tracker Properties onto the exposed properties of the current Preset
+	 * (Preset --> Tracker)
+	 */
+	void WriteAllPropertyIdsToPreset();
+
+	/** returns true if Tracker Component has at least one tracked/exposed property */
+	bool HasTrackedProperties() const;
+
+	/** returns true if Tracker Component is tracking the property for the specified FieldPathInfo and Owner Object */
+	bool IsTrackingProperty(const FRCFieldPathInfo& InFieldPathInfo, UObject* InOwnerObject) const;
+	
+	/** Get the Actor owning this Tracker Component */
+	AActor* GetTrackedActor() const;
+
+protected:
+	URemoteControlPreset* GetCurrentPreset() const;
+	
+	//~ Begin UActorComponent Interface
+	virtual void OnComponentCreated() override;
+	virtual void OnComponentDestroyed(bool bInDestroyingHierarchy) override;
+	//~ End UActorComponent
+	
+	//~ Begin UObject Interface
+	virtual void PostInitProperties() override;
+	virtual void PostDuplicate(bool bInDuplicateForPIE) override;
+	virtual void PostEditImport() override;
+	virtual void PostLoad() override;
+#if WITH_EDITOR
+	virtual void PostTransacted(const FTransactionObjectEvent& InTransactionEvent) override;
+#endif
+	//~ End UObject
+
+	void RegisterTrackedActor() const;
+	void UnregisterTrackedActor() const;
+
+	void OnTrackerDuplicated();
+	
+	void RefreshTracker();
+	void RefreshExposedProperties();
+	void MarkPropertiesForRefresh();
+	
+	void RegisterPropertyIdChangeDelegate();
+	void UnregisterPropertyIdChangeDelegate() const;
+	void OnPropertyIdUpdated();
+
+	int32 GetTrackedPropertyIndex(const FRCFieldPathInfo& InFieldPathInfo, UObject* InOwnerObject) const;
+
+private:
+	UPROPERTY()
+	TArray<FRemoteControlTrackerProperty> TrackedProperties;
+};
diff --git a/Engine/Plugins/Experimental/RemoteControlComponents/Source/RemoteControlComponents/Public/RemoteControlTrackerProperty.h b/Engine/Plugins/Experimental/RemoteControlComponents/Source/RemoteControlComponents/Public/RemoteControlTrackerProperty.h
new file mode 100644
--- /dev/null
+++ b/Engine/Plugins/Experimental/RemoteControlComponents/Source/RemoteControlComponents/Public/RemoteControlTrackerProperty.h
@@ -0,0 +1,100 @@
+// Copyright Epic Games, Inc. All Rights Reserved.
+
+#pragma once
+
+#include "RemoteControlFieldPath.h"
+#include "RemoteControlTrackerProperty.generated.h"
+
+class URemoteControlPreset;
+
+/**
+ * Represents an exposed Remote Control Property tracked by a Remote Control Tracker Component
+ * Its Field Path Info and Owner are used to represent the RC Property.
+ */
+USTRUCT()
+struct FRemoteControlTrackerProperty
+{
+	GENERATED_BODY()
+
+public:
+	FRemoteControlTrackerProperty() {}
+
+	/**
+	 * Creates and initializes a Tracked Property
+	 * @param InFieldPathInfo: Field Path Info to reach the property
+	 * @param InOwnerObject: the object owning the Tracked Property
+	 * @param bInIsExposed: is the property already exposed? Default is true.
+	 */
+	FRemoteControlTrackerProperty(const FRCFieldPathInfo& InFieldPathInfo, const TObjectPtr<UObject>& InOwnerObject, bool bInIsExposed = true);
+
+	/** Checks if this Tracked Property is matching the one described by the specified parameters */
+	bool MatchesParameters(const FRCFieldPathInfo& InFieldPathInfo, const TObjectPtr<UObject>& InOwnerObject) const;
+	
+	/** Returns the Remote Control Preset where this property is or will be exposed */
+	URemoteControlPreset* GetPreset() const;
+
+	/** Is the tracked property exposed to the specified Preset? */
+	bool IsExposedTo(const URemoteControlPreset* InPreset) const;
+	
+	/** Is this Tracked Property exposed? */
+	bool IsExposed() const { return bIsExposed; }
+
+	/** Expose this Tracked Property to the specified Preset */
+	void Expose(URemoteControlPreset* InRemoteControlPreset);
+
+	/** Unexpose this Tracked Property from its Preset */
+	void Unexpose();
+
+	/** Retrieves the current Property Id for the currently tracked property */
+	void ReadPropertyIdFromPreset();
+
+	/** Writes the Property Id saved in this Tracker Property to the tracked property in the Preset */
+	void WritePropertyIdToPreset() const;
+
+	/**
+	 * Marks this Tracked Property as not exposed, even if it might be.
+	 * Useful to force refresh the state of a Tracked Property.
+	 */
+	void MarkUnexposed();
+
+	bool operator==(const FRemoteControlTrackerProperty& Other) const
+	{
+		return
+			FieldPathInfo == Other.FieldPathInfo &&
+			OwnerObject == Other.OwnerObject;
+	}
+
+	bool operator!=(const FRemoteControlTrackerProperty& Other) const
+	{
+		return !(*this == Other);
+	}
+
+	friend uint32 GetTypeHash(const FRemoteControlTrackerProperty& InBroadcastControlId)
+	{
+		return HashCombineFast(InBroadcastControlId.FieldPathInfo.PathHash, GetTypeHash(InBroadcastControlId.OwnerObject));
+	}
+
+private:
+	/** Resolves the Tracked Property FieldPathInfo with its Owner */
+	void Resolve() const;
+
+	/** Field Path Info, matching the same one found in the Remote Control Preset */
+	UPROPERTY()
+	FRCFieldPathInfo FieldPathInfo;
+
+	/** The Property Id of the exposed property tracked by this Tracker Property */
+	UPROPERTY()
+	FName PropertyId = NAME_None;
+
+	/** Owner of the property handled by this Tracked Component */
+	UPROPERTY()
+	TWeakObjectPtr<UObject> OwnerObject = nullptr;
+
+	/** Is this Tracked Property exposed? */
+	UPROPERTY()
+	bool bIsExposed = false;
+
+	/** Caching the current Remote Control Preset */
+	UPROPERTY()
+	TWeakObjectPtr<URemoteControlPreset> CurrentPresetWeak = nullptr;
+};
diff --git a/Engine/Plugins/Experimental/RemoteControlComponents/Source/RemoteControlComponents/Public/Subsystems/RemoteControlComponentsSubsystem.h b/Engine/Plugins/Experimental/RemoteControlComponents/Source/RemoteControlComponents/Public/Subsystems/RemoteControlComponentsSubsystem.h
new file mode 100644
--- /dev/null
+++ b/Engine/Plugins/Experimental/RemoteControlComponents/Source/RemoteControlComponents/Public/Subsystems/RemoteControlComponentsSubsystem.h
@@ -0,0 +1,94 @@
+// Copyright Epic Games, Inc. All Rights Reserved.
+
+#pragma once
+
+#include "RemoteControlComponentsContext.h"
+#include "RemoteControlPreset.h"
+#include "Subsystems/EngineSubsystem.h"
+#include "Templates/SharedPointer.h"
+#include "UObject/ObjectKey.h"
+#include "RemoteControlComponentsSubsystem.generated.h"
+
+class FSubsystemCollectionBase;
+class UObject;
+class URemoteControlPreset;
+class URemoteControlTrackerComponent;
+
+struct FRCFieldPathInfo;
+
+UCLASS()
+class REMOTECONTROLCOMPONENTS_API URemoteControlComponentsSubsystem : public UEngineSubsystem
+{
+	GENERATED_BODY()
+
+public:
+	/** Get this subsystem instance */
+	static URemoteControlComponentsSubsystem* Get();
+
+	/** Register the specified Preset, along with the World it resides in. */
+	void RegisterPreset(URemoteControlPreset* InRemoteControlPreset);
+
+	/** Unregister the specified Preset */
+	void UnregisterPreset(URemoteControlPreset* InRemoteControlPreset);
+
+	/** Marks the specified Actor as a Tracked Actor, so that the Subsystem knows about it */
+	void RegisterTrackedActor(AActor* InActor);
+	
+	/** Removes the specified Actor from the tracked actors register */
+	void UnregisterTrackedActor(AActor* InActor);
+	
+	/** Checks if the specified Actor is currently being tracked by the subsystem */
+	bool IsActorTracked(const AActor* InActor) const;
+
+	/** Checks if the specified Preset is currently registered */
+	bool IsPresetRegistered(const URemoteControlPreset* InPreset) const;
+
+	/** Returns the Preset currently associated with the specified World */
+	URemoteControlPreset* GetRegisteredPreset(const UWorld* InWorld) const;
+
+	/** Returns the Preset currently associated with the specified Object */
+	URemoteControlPreset* GetRegisteredPreset(const UObject* InObject) const;
+
+	DECLARE_MULTICAST_DELEGATE_OneParam(FOnActorRegisterStateUpdate, AActor* /* InActor */)
+	/** Fires whenever an actor is registered for Tracking */
+	FOnActorRegisterStateUpdate& OnTrackedActorRegistered() { return OnTrackedActorRegisteredDelegate; }
+
+	/** Fires whenever an actor is unregistered from Tracking */
+	FOnActorRegisterStateUpdate& OnTrackedActorUnregistered() { return OnTrackedActorUnregisteredDelegate; }
+
+	DECLARE_MULTICAST_DELEGATE_OneParam(FOnActivePresetChanged, URemoteControlPreset* /* InActor */)
+	/** Fires when the active Preset is updated */
+	FOnActivePresetChanged& OnActivePresetChanged() { return OnActivePresetChangedDelegate; }
+
+protected:
+	FOnActorRegisterStateUpdate OnTrackedActorRegisteredDelegate;	
+	FOnActorRegisterStateUpdate OnTrackedActorUnregisteredDelegate;
+	FOnActivePresetChanged OnActivePresetChangedDelegate;
+
+	//~ Begin USubsystem
+	virtual void Initialize(FSubsystemCollectionBase& InCollection) override;
+	virtual void Deinitialize() override;
+	//~ End USubsystem
+
+	void OnLevelChanged(ULevel* InNewLevel, ULevel* InOldLevel, UWorld* InWorld);
+	void OnWorldCleanup(UWorld* InWorld, bool bInSessionEnded, bool bInCleanupResources);
+	void OnLevelActorDeleted(AActor* InActor);
+
+	void ClearData();
+	
+	void OnEntityExposed(URemoteControlPreset* InRemoteControlPreset, const FGuid& InGuid);
+	void OnEntityUnexposed(URemoteControlPreset* InRemoteControlPreset, const FGuid& InGuid);
+
+	void CacheActorsToBeTracked(const TSet<TWeakObjectPtr<AActor>>& InActors);
+	void RegisterPendingActors(const UWorld* InWorld);
+
+	bool IsRegisteredActor(const AActor* InActor) const;
+	bool IsRegisteredTrackedActor(const AActor* InActor) const;
+
+	TSharedPtr<FRemoteControlComponentsContext> GetTrackedActorContext(const UWorld* InWorld) const;
+
+	/** Utility container used to store actors for later tracking in case registering doesn't work right away */
+	TArray<TWeakObjectPtr<AActor>> ActorsPendingRegister;
+
+	TMap<FObjectKey, TSharedPtr<FRemoteControlComponentsContext>> TrackedActorsContextMap;
+};
diff --git a/Engine/Plugins/Experimental/RemoteControlComponents/Source/RemoteControlComponents/RemoteControlComponents.Build.cs b/Engine/Plugins/Experimental/RemoteControlComponents/Source/RemoteControlComponents/RemoteControlComponents.Build.cs
new file mode 100644
--- /dev/null
+++ b/Engine/Plugins/Experimental/RemoteControlComponents/Source/RemoteControlComponents/RemoteControlComponents.Build.cs
@@ -0,0 +1,38 @@
+﻿// Copyright Epic Games, Inc. All Rights Reserved.
+
+using UnrealBuildTool;
+
+public class RemoteControlComponents : ModuleRules
+{
+	public RemoteControlComponents(ReadOnlyTargetRules Target) : base(Target)
+	{
+		PCHUsage = ModuleRules.PCHUsageMode.UseExplicitOrSharedPCHs;
+		
+		PublicDependencyModuleNames.AddRange(
+			new string[]
+			{
+				"Core",
+				"CoreUObject",
+				"RemoteControl"
+			}
+			);
+	
+
+		PrivateDependencyModuleNames.AddRange(
+			new string[]
+			{
+				"Engine",
+			}
+			);
+
+
+		if (Target.Type == TargetRules.TargetType.Editor)
+		{
+			PrivateDependencyModuleNames.AddRange(new string[]
+			{
+				"RemoteControlUI",
+				"UnrealEd",
+			});
+		}
+	}
+}
diff --git a/Engine/Plugins/Experimental/RemoteControlComponents/Source/RemoteControlComponentsEditor/Private/RemoteControlComponentsEditorModule.cpp b/Engine/Plugins/Experimental/RemoteControlComponents/Source/RemoteControlComponentsEditor/Private/RemoteControlComponentsEditorModule.cpp
new file mode 100644
--- /dev/null
+++ b/Engine/Plugins/Experimental/RemoteControlComponents/Source/RemoteControlComponentsEditor/Private/RemoteControlComponentsEditorModule.cpp
@@ -0,0 +1,5 @@
+// Copyright Epic Games, Inc. All Rights Reserved.
+
+#include "Modules/ModuleManager.h"
+
+IMPLEMENT_MODULE(FDefaultModuleImpl, RemoteControlComponentsEditor)
diff --git a/Engine/Plugins/Experimental/RemoteControlComponents/Source/RemoteControlComponentsEditor/Private/Styles/RemoteControlComponentsEditorStyle.cpp b/Engine/Plugins/Experimental/RemoteControlComponents/Source/RemoteControlComponentsEditor/Private/Styles/RemoteControlComponentsEditorStyle.cpp
new file mode 100644
--- /dev/null
+++ b/Engine/Plugins/Experimental/RemoteControlComponents/Source/RemoteControlComponentsEditor/Private/Styles/RemoteControlComponentsEditorStyle.cpp
@@ -0,0 +1,28 @@
+// Copyright Epic Games, Inc. All Rights Reserved.
+
+#include "RemoteControlComponentsEditorStyle.h"
+#include "Interfaces/IPluginManager.h"
+#include "Misc/Paths.h"
+#include "Styling/SlateStyleMacros.h"
+#include "Styling/SlateStyleRegistry.h"
+
+FRemoteControlComponentsEditorStyle::FRemoteControlComponentsEditorStyle()
+	: FSlateStyleSet(TEXT("RemoteControlComponentsEditor"))
+{
+	const FVector2f Icon16x16(16.f, 16.f);
+	
+	const TSharedPtr<IPlugin> Plugin = IPluginManager::Get().FindPlugin(UE_PLUGIN_NAME);
+	check(Plugin.IsValid());
+	
+	ContentRootDir = FPaths::Combine(Plugin->GetBaseDir(), TEXT("Resources"));
+
+	Set("ClassIcon.RemoteControlComponentsComponent"    , new IMAGE_BRUSH_SVG("EditorIcons/RemoteControlComponentsDefault", Icon16x16));
+	Set("ClassIcon.RemoteControlTracker"                , new IMAGE_BRUSH_SVG("EditorIcons/RemoteControlTracker", Icon16x16));
+
+	FSlateStyleRegistry::RegisterSlateStyle(*this);
+}
+
+FRemoteControlComponentsEditorStyle::~FRemoteControlComponentsEditorStyle()
+{
+	FSlateStyleRegistry::UnRegisterSlateStyle(*this);
+}
diff --git a/Engine/Plugins/Experimental/RemoteControlComponents/Source/RemoteControlComponentsEditor/Private/Styles/RemoteControlComponentsEditorStyle.h b/Engine/Plugins/Experimental/RemoteControlComponents/Source/RemoteControlComponentsEditor/Private/Styles/RemoteControlComponentsEditorStyle.h
new file mode 100644
--- /dev/null
+++ b/Engine/Plugins/Experimental/RemoteControlComponents/Source/RemoteControlComponentsEditor/Private/Styles/RemoteControlComponentsEditorStyle.h
@@ -0,0 +1,19 @@
+// Copyright Epic Games, Inc. All Rights Reserved.
+
+#pragma once
+
+#include "Styling/SlateStyle.h"
+
+class FRemoteControlComponentsEditorStyle : public FSlateStyleSet
+{
+public:
+	FRemoteControlComponentsEditorStyle();
+	
+	virtual ~FRemoteControlComponentsEditorStyle() override;
+
+	static FRemoteControlComponentsEditorStyle& Get()
+	{
+		static FRemoteControlComponentsEditorStyle StyleSet;
+		return StyleSet;
+	}
+};
diff --git a/Engine/Plugins/Experimental/RemoteControlComponents/Source/RemoteControlComponentsEditor/Private/Subsystems/RemoteControlComponentsEditorUtils.cpp b/Engine/Plugins/Experimental/RemoteControlComponents/Source/RemoteControlComponentsEditor/Private/Subsystems/RemoteControlComponentsEditorUtils.cpp
new file mode 100644
--- /dev/null
+++ b/Engine/Plugins/Experimental/RemoteControlComponents/Source/RemoteControlComponentsEditor/Private/Subsystems/RemoteControlComponentsEditorUtils.cpp
@@ -0,0 +1,16 @@
+// Copyright Epic Games, Inc. All Rights Reserved.
+
+#include "Subsystems/RemoteControlComponentsEditorUtils.h"
+#include "Editor/EditorEngine.h"
+#include "Framework/Application/SlateApplication.h"
+#include "Styles/RemoteControlComponentsEditorStyle.h"
+#include "Widgets/Docking/SDockTab.h"
+
+#define LOCTEXT_NAMESPACE "RemoteControlComponentsEditorUtils"
+
+FSlateIcon FRemoteControlComponentsEditorUtils::GetIcon(const FName& InIconName)
+{
+	return FSlateIcon(FRemoteControlComponentsEditorStyle::Get().GetStyleSetName(), InIconName);
+}
+
+#undef LOCTEXT_NAMESPACE
diff --git a/Engine/Plugins/Experimental/RemoteControlComponents/Source/RemoteControlComponentsEditor/Public/Subsystems/RemoteControlComponentsEditorUtils.h b/Engine/Plugins/Experimental/RemoteControlComponents/Source/RemoteControlComponentsEditor/Public/Subsystems/RemoteControlComponentsEditorUtils.h
new file mode 100644
--- /dev/null
+++ b/Engine/Plugins/Experimental/RemoteControlComponents/Source/RemoteControlComponentsEditor/Public/Subsystems/RemoteControlComponentsEditorUtils.h
@@ -0,0 +1,15 @@
+// Copyright Epic Games, Inc. All Rights Reserved.
+
+#pragma once
+
+#include "HAL/Platform.h"
+
+class FName;
+struct FSlateIcon;
+
+class FRemoteControlComponentsEditorUtils
+{
+public:
+	/** Retrieve the specified Icon from the Style associated to the Remote Control Components Editor */
+	REMOTECONTROLCOMPONENTSEDITOR_API static FSlateIcon GetIcon(const FName& InIconName);
+};
diff --git a/Engine/Plugins/Experimental/RemoteControlComponents/Source/RemoteControlComponentsEditor/RemoteControlComponentsEditor.Build.cs b/Engine/Plugins/Experimental/RemoteControlComponents/Source/RemoteControlComponentsEditor/RemoteControlComponentsEditor.Build.cs
new file mode 100644
--- /dev/null
+++ b/Engine/Plugins/Experimental/RemoteControlComponents/Source/RemoteControlComponentsEditor/RemoteControlComponentsEditor.Build.cs
@@ -0,0 +1,33 @@
+// Copyright Epic Games, Inc. All Rights Reserved.
+
+using UnrealBuildTool;
+
+public class RemoteControlComponentsEditor : ModuleRules
+{
+    public RemoteControlComponentsEditor(ReadOnlyTargetRules Target) : base(Target)
+    {
+        PCHUsage = ModuleRules.PCHUsageMode.UseExplicitOrSharedPCHs;
+
+        PublicDependencyModuleNames.AddRange(
+            new string[]
+            {
+                "Core"
+            }
+        );
+
+        PrivateDependencyModuleNames.AddRange(
+            new string[]
+            {
+                "CoreUObject",
+                "Engine",
+                "EditorSubsystem",
+                "Projects",
+                "RemoteControl",
+				"RemoteControlComponents",
+                "Slate",
+                "SlateCore",
+                "UnrealEd"
+            }
+        );
+    }
+}
\ No newline at end of file
