Remote Control Components:
- Cleanup pass unexporting functions that dont need to be
- Fix issue with icon style not properly set up for Tracker Component

#changelist validated
#virtualized

diff --git a/Engine/Plugins/Experimental/Avalanche/Source/AvalancheEditor/Private/RemoteControl/AvaRCExtension.cpp b/Engine/Plugins/Experimental/Avalanche/Source/AvalancheEditor/Private/RemoteControl/AvaRCExtension.cpp
--- a/Engine/Plugins/Experimental/Avalanche/Source/AvalancheEditor/Private/RemoteControl/AvaRCExtension.cpp
+++ b/Engine/Plugins/Experimental/Avalanche/Source/AvalancheEditor/Private/RemoteControl/AvaRCExtension.cpp
@@ -5,6 +5,7 @@
 #include "Framework/Docking/LayoutExtender.h"
 #include "IAvaSceneInterface.h"
 #include "LevelEditor.h"
+#include "RemoteControlPreset.h"
 #include "RemoteControlTrackerComponent.h"
 #include "Subsystems/AssetEditorSubsystem.h"
 #include "Subsystems/RemoteControlComponentsSubsystem.h"
diff --git a/Engine/Plugins/Experimental/Avalanche/Source/AvalancheRemoteControlEditor/Private/Outliner/AvaOutlinerRCTrackerComponent.cpp b/Engine/Plugins/Experimental/Avalanche/Source/AvalancheRemoteControlEditor/Private/Outliner/AvaOutlinerRCTrackerComponent.cpp
--- a/Engine/Plugins/Experimental/Avalanche/Source/AvalancheRemoteControlEditor/Private/Outliner/AvaOutlinerRCTrackerComponent.cpp
+++ b/Engine/Plugins/Experimental/Avalanche/Source/AvalancheRemoteControlEditor/Private/Outliner/AvaOutlinerRCTrackerComponent.cpp
@@ -6,8 +6,8 @@
 #include "RemoteControlPreset.h"
 #include "RemoteControlTrackerComponent.h"
 #include "Selection/AvaOutlinerScopedSelection.h"
+#include "Styling/SlateIconFinder.h"
 #include "Subsystems/AssetEditorSubsystem.h"
-#include "Subsystems/RemoteControlComponentsEditorUtils.h"
 
 #define LOCTEXT_NAMESPACE "AvaOutlinerRemoteControlComponent"
 
@@ -32,7 +32,7 @@ FAvaOutlinerRCTrackerComponent::FAvaOutlinerRCTrackerComponent(IAvaOutliner& InO
 	: FAvaOutlinerObject(InOutliner, InComponent)
 	, TrackerComponentWeak(InComponent)
 {
-	TrackerIcon = FRemoteControlComponentsEditorUtils::GetIcon("ClassIcon.RemoteControlTracker");
+	TrackerIcon = FSlateIconFinder::FindIconForClass(URemoteControlTrackerComponent::StaticClass());
 }
 
 FSlateIcon FAvaOutlinerRCTrackerComponent::GetIcon() const
diff --git a/Engine/Plugins/Experimental/Avalanche/Source/AvalancheRemoteControlEditor/Private/Outliner/AvaOutlinerRCTrackerComponentProxy.cpp b/Engine/Plugins/Experimental/Avalanche/Source/AvalancheRemoteControlEditor/Private/Outliner/AvaOutlinerRCTrackerComponentProxy.cpp
--- a/Engine/Plugins/Experimental/Avalanche/Source/AvalancheRemoteControlEditor/Private/Outliner/AvaOutlinerRCTrackerComponentProxy.cpp
+++ b/Engine/Plugins/Experimental/Avalanche/Source/AvalancheRemoteControlEditor/Private/Outliner/AvaOutlinerRCTrackerComponentProxy.cpp
@@ -6,7 +6,7 @@
 #include "Item/AvaOutlinerActor.h"
 #include "Outliner/AvaOutlinerRCTrackerComponent.h"
 #include "RemoteControlTrackerComponent.h"
-#include "Subsystems/RemoteControlComponentsEditorUtils.h"
+#include "Styling/SlateIconFinder.h"
 #include "Subsystems/RemoteControlComponentsSubsystem.h"
 
 #define LOCTEXT_NAMESPACE "AvaOutlinerRemoteControlComponentProxy"
@@ -14,7 +14,7 @@
 FAvaOutlinerRCTrackerComponentProxy::FAvaOutlinerRCTrackerComponentProxy(IAvaOutliner& InOutliner, const FAvaOutlinerItemPtr& InParentItem)
 	: Super(InOutliner, InParentItem)
 {
-	TrackerIcon = FRemoteControlComponentsEditorUtils::GetIcon("ClassIcon.RemoteControlTracker");
+	TrackerIcon = FSlateIconFinder::FindIconForClass(URemoteControlTrackerComponent::StaticClass());
 }
 
 URemoteControlTrackerComponent* FAvaOutlinerRCTrackerComponentProxy::GetTrackerComponent() const
diff --git a/Engine/Plugins/Experimental/RemoteControlComponents/Source/RemoteControlComponents/Private/RemoteControlComponentsContext.h b/Engine/Plugins/Experimental/RemoteControlComponents/Source/RemoteControlComponents/Private/RemoteControlComponentsContext.h
new file mode 100644
--- /dev/null
+++ b/Engine/Plugins/Experimental/RemoteControlComponents/Source/RemoteControlComponents/Private/RemoteControlComponentsContext.h
@@ -0,0 +1,62 @@
+﻿// Copyright Epic Games, Inc. All Rights Reserved.
+
+#pragma once
+
+#include "Containers/Array.h"
+#include "UObject/WeakObjectPtr.h"
+
+class AActor;
+class URemoteControlPreset;
+class UWorld;
+
+/**
+ * Remote Control Components Tracking Context: represents a World with a set of Actors which have properties exposed to a specific Preset.
+ * It currently works under the assumption that there's a single Preset per World, and that this info is runtime only. Both aspects might change in the future.
+ *
+ * Context is represented by:
+ *	- a World
+ *	- a Remote Control Preset
+ *	- Tracked Actors with properties exposed to the Preset
+ */
+struct FRemoteControlComponentsContext
+{
+public:
+	/** Create a Remote Control Context by passing a World and a Preset */
+	FRemoteControlComponentsContext(UWorld* InWorld, URemoteControlPreset* InPreset);
+
+	/** Adds an Actor to the list of tracked actors */
+	bool RegisterActor(AActor* InActor);
+
+	/** Removes an Actor from the list of tracked actors */
+	bool UnregisterActor(AActor* InActor);
+
+	/** Checks if the specified Actor is tracked or not */
+	bool IsActorRegistered(const AActor* InActor) const;
+
+	/** An Actor can be registered if it has a Remote Control Tracker Component and its World and the World of this context coincide */
+	bool CanRegisterActor(const AActor* InActor) const;
+
+	/** Context is valid if both World and Preset are valid */
+	bool IsValid() const;
+
+	/** True if the context currently has registered tracked actors */
+	bool HasTrackedActors() const;
+
+	/** Get the array of currently tracked actors */
+	TArray<TWeakObjectPtr<AActor>> GetTrackedActors() const { return TrackedActors; }
+
+	/** Get the World referenced by this context */
+	UWorld* GetWorld() const;
+
+	/** Get the Remote Control Preset referenced by this context */
+	URemoteControlPreset* GetPreset() const;
+
+private:
+	TWeakObjectPtr<UWorld> WorldWeak;
+	
+	TWeakObjectPtr<URemoteControlPreset> PresetWeak;
+	
+	TArray<TWeakObjectPtr<AActor>> TrackedActors;
+
+	bool bIsValid = false;
+};
diff --git a/Engine/Plugins/Experimental/RemoteControlComponents/Source/RemoteControlComponents/Private/RemoteControlTrackerComponent.cpp b/Engine/Plugins/Experimental/RemoteControlComponents/Source/RemoteControlComponents/Private/RemoteControlTrackerComponent.cpp
--- a/Engine/Plugins/Experimental/RemoteControlComponents/Source/RemoteControlComponents/Private/RemoteControlTrackerComponent.cpp
+++ b/Engine/Plugins/Experimental/RemoteControlComponents/Source/RemoteControlComponents/Private/RemoteControlTrackerComponent.cpp
@@ -8,42 +8,101 @@
 #include "RemoteControlPreset.h"
 #include "Subsystems/RemoteControlComponentsSubsystem.h"
 
-void URemoteControlTrackerComponent::MarkPropertiesForRefresh()
+URemoteControlPreset* URemoteControlTrackerComponent::GetCurrentPreset() const
 {
-	for (FRemoteControlTrackerProperty& TrackedProperty : TrackedProperties)
+	if (const URemoteControlComponentsSubsystem* RemoteControlPresetComponentsSubsystem = URemoteControlComponentsSubsystem::Get())
 	{
-		TrackedProperty.MarkUnexposed();
+		if (const UWorld* World = GetTypedOuter<UWorld>())
+		{
+			if (URemoteControlPreset* Preset = RemoteControlPresetComponentsSubsystem->GetRegisteredPreset(World))
+			{
+				return Preset;
+			}
+		}
+	}
+
+	return nullptr;
+}
+
+bool URemoteControlTrackerComponent::HasTrackedProperties() const
+{
+	return !TrackedProperties.IsEmpty();
+}
+
+void URemoteControlTrackerComponent::AddTrackedProperty(const FRCFieldPathInfo& InFieldPathInfo, UObject* InOwnerObject)
+{
+	if (!IsTrackingProperty(InFieldPathInfo, InOwnerObject))
+	{
+		TrackedProperties.Add(FRemoteControlTrackerProperty(InFieldPathInfo, InOwnerObject));
 	}
 }
 
-void URemoteControlTrackerComponent::RefreshExposedProperties()
+void URemoteControlTrackerComponent::RemoveTrackedProperty(const FRCFieldPathInfo& InFieldPathInfo, UObject* InOwnerObject)
 {
-	MarkPropertiesForRefresh();
-	ExposeAllProperties();
+	const int32 Index = GetTrackedPropertyIndex(InFieldPathInfo, InOwnerObject);
+	if (Index != INDEX_NONE)
+	{
+		TrackedProperties.RemoveAt(Index);
+	}
 }
 
-void URemoteControlTrackerComponent::RegisterTrackedActor() const
-{	
-	if (AActor* OuterActor = GetTrackedActor())
+void URemoteControlTrackerComponent::ExposeAllProperties()
+{
+	if (URemoteControlPreset* Preset = GetCurrentPreset())
 	{
-		if (URemoteControlComponentsSubsystem* RemoteControlComponentsSubsystem = URemoteControlComponentsSubsystem::Get())
+		for (FRemoteControlTrackerProperty& TrackedProperty : TrackedProperties)
 		{
-			RemoteControlComponentsSubsystem->RegisterTrackedActor(OuterActor);
+			TrackedProperty.Expose(Preset);
 		}
 	}
 }
 
-void URemoteControlTrackerComponent::UnregisterTrackedActor() const
+void URemoteControlTrackerComponent::UnexposeAllProperties()
 {
-	if (AActor* OuterActor = GetTrackedActor())
+	TArray<FRemoteControlTrackerProperty> PropertiesToRemove = MoveTemp(TrackedProperties);
+
+	for (FRemoteControlTrackerProperty& TrackedProperty : PropertiesToRemove)
 	{
-		if (URemoteControlComponentsSubsystem* RemoteControlComponentsSubsystem = URemoteControlComponentsSubsystem::Get())
-		{
-			RemoteControlComponentsSubsystem->UnregisterTrackedActor(OuterActor);
-		}
+		TrackedProperty.Unexpose();
+	}
+	
+	UnregisterPropertyIdChangeDelegate();
+}
+
+void URemoteControlTrackerComponent::RefreshAllPropertyIds()
+{
+	for (FRemoteControlTrackerProperty& TrackedProperty : TrackedProperties)
+	{
+		TrackedProperty.ReadPropertyIdFromPreset();
 	}
 }
 
+void URemoteControlTrackerComponent::WriteAllPropertyIdsToPreset()
+{
+	UnregisterPropertyIdChangeDelegate();
+	for (FRemoteControlTrackerProperty& TrackedProperty : TrackedProperties)
+	{
+		TrackedProperty.WritePropertyIdToPreset();
+	}
+	RegisterPropertyIdChangeDelegate();
+}
+
+bool URemoteControlTrackerComponent::IsTrackingProperty(const FRCFieldPathInfo& InFieldPathInfo, UObject* InOwnerObject) const
+{
+	const int32 Index = GetTrackedPropertyIndex(InFieldPathInfo, InOwnerObject);
+	return Index != INDEX_NONE;
+}
+
+AActor* URemoteControlTrackerComponent::GetTrackedActor() const
+{
+	if (AActor* OuterActor = GetOwner())
+	{
+		return OuterActor;
+	}
+
+	return nullptr;
+}
+
 void URemoteControlTrackerComponent::OnComponentCreated()
 {
 	Super::OnComponentCreated();
@@ -72,12 +131,6 @@
 	}
 }
 
-void URemoteControlTrackerComponent::RefreshTracker()
-{
-	RegisterTrackedActor();
-	RefreshExposedProperties();
-}
-
 void URemoteControlTrackerComponent::PostInitProperties()
 {
 	Super::PostInitProperties();
@@ -86,13 +139,6 @@
 	RegisterPropertyIdChangeDelegate();
 }
 
-void URemoteControlTrackerComponent::OnTrackerDuplicated()
-{
-	RefreshTracker();
-	WriteAllPropertyIdsToPreset();
-	RegisterPropertyIdChangeDelegate();
-}
-
 void URemoteControlTrackerComponent::PostDuplicate(bool bInDuplicateForPIE)
 {
 	Super::PostDuplicate(bInDuplicateForPIE);
@@ -131,107 +177,53 @@
 }
 #endif
 
-int32 URemoteControlTrackerComponent::GetTrackedPropertyIndex(const FRCFieldPathInfo& InFieldPathInfo, UObject* InOwnerObject) const
-{
-	return TrackedProperties.IndexOfByPredicate([InFieldPathInfo, InOwnerObject](const FRemoteControlTrackerProperty& TrackerProperty)
+void URemoteControlTrackerComponent::RegisterTrackedActor() const
+{	
+	if (AActor* OuterActor = GetTrackedActor())
 	{
-		return TrackerProperty.MatchesParameters(InFieldPathInfo, InOwnerObject);
-	});
-}
-
-void URemoteControlTrackerComponent::AddTrackedProperty(const FRCFieldPathInfo& InFieldPathInfo, UObject* InOwnerObject)
-{
-	if (!IsTrackingProperty(InFieldPathInfo, InOwnerObject))
-	{
-		TrackedProperties.Add(FRemoteControlTrackerProperty(InFieldPathInfo, InOwnerObject));
-	}
-}
-
-void URemoteControlTrackerComponent::RemoveTrackedProperty(const FRCFieldPathInfo& InFieldPathInfo, UObject* InOwnerObject)
-{
-	const int32 Index = GetTrackedPropertyIndex(InFieldPathInfo, InOwnerObject);
-	if (Index != INDEX_NONE)
-	{
-		TrackedProperties.RemoveAt(Index);
-	}
-}
-
-URemoteControlPreset* URemoteControlTrackerComponent::GetCurrentPreset() const
-{
-	if (const URemoteControlComponentsSubsystem* RemoteControlPresetComponentsSubsystem = URemoteControlComponentsSubsystem::Get())
-	{
-		if (const UWorld* World = GetTypedOuter<UWorld>())
+		if (URemoteControlComponentsSubsystem* RemoteControlComponentsSubsystem = URemoteControlComponentsSubsystem::Get())
 		{
-			if (URemoteControlPreset* Preset = RemoteControlPresetComponentsSubsystem->GetRegisteredPreset(World))
-			{
-				return Preset;
-			}
-		}
-	}
-
-	return nullptr;
-}
-
-void URemoteControlTrackerComponent::ExposeAllProperties()
-{
-	if (URemoteControlPreset* Preset = GetCurrentPreset())
-	{
-		for (FRemoteControlTrackerProperty& TrackedProperty : TrackedProperties)
-		{
-			TrackedProperty.Expose(Preset);
+			RemoteControlComponentsSubsystem->RegisterTrackedActor(OuterActor);
 		}
 	}
 }
 
-bool URemoteControlTrackerComponent::HasTrackedProperties() const
+void URemoteControlTrackerComponent::UnregisterTrackedActor() const
 {
-	return !TrackedProperties.IsEmpty();
-}
-
-bool URemoteControlTrackerComponent::IsTrackingProperty(const FRCFieldPathInfo& InFieldPathInfo, UObject* InOwnerObject) const
-{
-	const int32 Index = GetTrackedPropertyIndex(InFieldPathInfo, InOwnerObject);
-	return Index != INDEX_NONE;
-}
-
-void URemoteControlTrackerComponent::UnexposeAllProperties()
-{
-	TArray<FRemoteControlTrackerProperty> PropertiesToRemove = MoveTemp(TrackedProperties);
-
-	for (FRemoteControlTrackerProperty& TrackedProperty : PropertiesToRemove)
+	if (AActor* OuterActor = GetTrackedActor())
 	{
-		TrackedProperty.Unexpose();
-	}
-	
-	UnregisterPropertyIdChangeDelegate();
-}
-
-void URemoteControlTrackerComponent::RefreshAllPropertyIds()
-{
-	for (FRemoteControlTrackerProperty& TrackedProperty : TrackedProperties)
-	{
-		TrackedProperty.ReadPropertyIdFromPreset();
+		if (URemoteControlComponentsSubsystem* RemoteControlComponentsSubsystem = URemoteControlComponentsSubsystem::Get())
+		{
+			RemoteControlComponentsSubsystem->UnregisterTrackedActor(OuterActor);
+		}
 	}
 }
 
-void URemoteControlTrackerComponent::WriteAllPropertyIdsToPreset()
+void URemoteControlTrackerComponent::OnTrackerDuplicated()
 {
-	UnregisterPropertyIdChangeDelegate();
-	for (FRemoteControlTrackerProperty& TrackedProperty : TrackedProperties)
-	{
-		TrackedProperty.WritePropertyIdToPreset();
-	}
+	RefreshTracker();
+	WriteAllPropertyIdsToPreset();
 	RegisterPropertyIdChangeDelegate();
 }
 
-AActor* URemoteControlTrackerComponent::GetTrackedActor() const
+void URemoteControlTrackerComponent::RefreshTracker()
 {
-	if (AActor* OuterActor = GetOwner())
-	{
-		return OuterActor;
-	}
-
-	return nullptr;
+	RegisterTrackedActor();
+	RefreshExposedProperties();
 }
 
+void URemoteControlTrackerComponent::RefreshExposedProperties()
+{
+	MarkPropertiesForRefresh();
+	ExposeAllProperties();
+}
+
+void URemoteControlTrackerComponent::MarkPropertiesForRefresh()
+{
+	for (FRemoteControlTrackerProperty& TrackedProperty : TrackedProperties)
+	{
+		TrackedProperty.MarkUnexposed();
+	}
+}
+
 void URemoteControlTrackerComponent::RegisterPropertyIdChangeDelegate()
@@ -261,3 +253,11 @@
 {
 	RefreshAllPropertyIds();
 }
+
+int32 URemoteControlTrackerComponent::GetTrackedPropertyIndex(const FRCFieldPathInfo& InFieldPathInfo, UObject* InOwnerObject) const
+{
+	return TrackedProperties.IndexOfByPredicate([InFieldPathInfo, InOwnerObject](const FRemoteControlTrackerProperty& TrackerProperty)
+	{
+		return TrackerProperty.MatchesParameters(InFieldPathInfo, InOwnerObject);
+	});
+}
diff --git a/Engine/Plugins/Experimental/RemoteControlComponents/Source/RemoteControlComponents/Private/RemoteControlTrackerProperty.cpp b/Engine/Plugins/Experimental/RemoteControlComponents/Source/RemoteControlComponents/Private/RemoteControlTrackerProperty.cpp
--- a/Engine/Plugins/Experimental/RemoteControlComponents/Source/RemoteControlComponents/Private/RemoteControlTrackerProperty.cpp
+++ b/Engine/Plugins/Experimental/RemoteControlComponents/Source/RemoteControlComponents/Private/RemoteControlTrackerProperty.cpp
@@ -5,6 +5,10 @@
 #include "RemoteControlComponentsUtils.h"
 #include "RemoteControlPreset.h"
 
+FRemoteControlTrackerProperty::FRemoteControlTrackerProperty()
+{
+}
+
 FRemoteControlTrackerProperty::FRemoteControlTrackerProperty(const FRCFieldPathInfo& InFieldPathInfo, const TObjectPtr<UObject>& InOwnerObject, bool bInIsExposed)
 	: FieldPathInfo(InFieldPathInfo)
 	, OwnerObject(InOwnerObject)
@@ -51,6 +55,16 @@ void FRemoteControlTrackerProperty::MarkUnexposed()
 	bIsExposed = false;
 }
 
+bool FRemoteControlTrackerProperty::operator==(const FRemoteControlTrackerProperty& Other) const
+{
+	return FieldPathInfo == Other.FieldPathInfo && OwnerObject == Other.OwnerObject;
+}
+
+bool FRemoteControlTrackerProperty::operator!=(const FRemoteControlTrackerProperty& Other) const
+{
+	return !(*this == Other);
+}
+
 void FRemoteControlTrackerProperty::Expose(URemoteControlPreset* InRemoteControlPreset)
 {
 	if (IsExposedTo(InRemoteControlPreset))
@@ -113,3 +127,8 @@ void FRemoteControlTrackerProperty::WritePropertyIdToPreset() const
 		FRemoteControlComponentsUtils::SetExposedPropertyId(Preset, OwnerObject.Get(), FieldPathInfo, PropertyId);
 	}
 }
+
+uint32 GetTypeHash(const FRemoteControlTrackerProperty& InBroadcastControlId)
+{
+	return HashCombineFast(InBroadcastControlId.FieldPathInfo.PathHash, GetTypeHash(InBroadcastControlId.OwnerObject));
+}
diff --git a/Engine/Plugins/Experimental/RemoteControlComponents/Source/RemoteControlComponents/Private/Subsystems/RemoteControlComponentsSubsystem.cpp b/Engine/Plugins/Experimental/RemoteControlComponents/Source/RemoteControlComponents/Private/Subsystems/RemoteControlComponentsSubsystem.cpp
--- a/Engine/Plugins/Experimental/RemoteControlComponents/Source/RemoteControlComponents/Private/Subsystems/RemoteControlComponentsSubsystem.cpp
+++ b/Engine/Plugins/Experimental/RemoteControlComponents/Source/RemoteControlComponents/Private/Subsystems/RemoteControlComponentsSubsystem.cpp
@@ -3,6 +3,7 @@
 #include "Subsystems/RemoteControlComponentsSubsystem.h"
 #include "Engine/Engine.h"
 #include "IRemoteControlModule.h"
+#include "RemoteControlComponentsContext.h"
 #include "RemoteControlComponentsUtils.h"
 #include "RemoteControlField.h"
 #include "RemoteControlPreset.h"
diff --git a/Engine/Plugins/Experimental/RemoteControlComponents/Source/RemoteControlComponents/Public/RemoteControlComponentsContext.h b/Engine/Plugins/Experimental/RemoteControlComponents/Source/RemoteControlComponents/Public/RemoteControlComponentsContext.h
deleted file mode 100644
diff --git a/Engine/Plugins/Experimental/RemoteControlComponents/Source/RemoteControlComponents/Public/RemoteControlComponentsUtils.h b/Engine/Plugins/Experimental/RemoteControlComponents/Source/RemoteControlComponents/Public/RemoteControlComponentsUtils.h
--- a/Engine/Plugins/Experimental/RemoteControlComponents/Source/RemoteControlComponents/Public/RemoteControlComponentsUtils.h
+++ b/Engine/Plugins/Experimental/RemoteControlComponents/Source/RemoteControlComponents/Public/RemoteControlComponentsUtils.h
@@ -17,6 +17,7 @@ class FRemoteControlComponentsUtils
 {
 	friend class URemoteControlComponentsSubsystem;
 	friend struct FRemoteControlTrackerProperty;
+
 public:
 	/** Add a RemoteControlTrackerComponent to the specified Actors */
 	REMOTECONTROLCOMPONENTS_API static void AddTrackerComponent(const TSet<TWeakObjectPtr<AActor>>& InActors, bool bInShouldTransact = true);
@@ -24,28 +25,28 @@ public:
 	/** Removes the RemoteControlTrackerComponent from the specified Actors */
 	REMOTECONTROLCOMPONENTS_API static void RemoveTrackerComponent(const TSet<TWeakObjectPtr<AActor>>& InActors);
 
+	/**
+	 * Unexpose all Tracked Properties of the specified Actor.
+	 * It will also stop tracking all unexposed properties.
+	 */
+	REMOTECONTROLCOMPONENTS_API static void UnexposeAllProperties(AActor* InActor);
+
 	/**
 	 * Gets the tracker component for the specified Actor. Tracker can be optionally added if missing.
 	 * @param InActor: the Actor for which we want to retrieve the Tracker Component. Object 
 	 * @param bInAddTrackerIfMissing: if true, function will add a tracker component when one is missing
 	 * @return The found or created tracker.
 	 */
-	REMOTECONTROLCOMPONENTS_API static URemoteControlTrackerComponent* GetTrackerComponent(AActor* InActor, bool bInAddTrackerIfMissing = false);
+	static URemoteControlTrackerComponent* GetTrackerComponent(AActor* InActor, bool bInAddTrackerIfMissing = false);
 
 	/** Gets the Remote Control Preset currently handling the specified Object */
-	REMOTECONTROLCOMPONENTS_API static URemoteControlPreset* GetCurrentPreset(const UObject* InObject);
+	static URemoteControlPreset* GetCurrentPreset(const UObject* InObject);
 
 	/** Gets the Remote Control Preset currently handling the specified World */
-	REMOTECONTROLCOMPONENTS_API static URemoteControlPreset* GetCurrentPreset(const UWorld* InWorld);
-
-	/**
-	 * Unexpose all Tracked Properties of the specified Actor.
-	 * It will also stop tracking all unexposed properties.
-	 */
-	REMOTECONTROLCOMPONENTS_API static void UnexposeAllProperties(AActor* InActor);
+	static URemoteControlPreset* GetCurrentPreset(const UWorld* InWorld);
 
 	/** Call this function on any Actor to ensure its exposed RC properties are actually tracked */
-	REMOTECONTROLCOMPONENTS_API static void RefreshTrackedProperties(AActor* InActor);
+	static void RefreshTrackedProperties(AActor* InActor);
 	
 private:
 	/**
diff --git a/Engine/Plugins/Experimental/RemoteControlComponents/Source/RemoteControlComponents/Public/RemoteControlTrackerComponent.h b/Engine/Plugins/Experimental/RemoteControlComponents/Source/RemoteControlComponents/Public/RemoteControlTrackerComponent.h
--- a/Engine/Plugins/Experimental/RemoteControlComponents/Source/RemoteControlComponents/Public/RemoteControlTrackerComponent.h
+++ b/Engine/Plugins/Experimental/RemoteControlComponents/Source/RemoteControlComponents/Public/RemoteControlTrackerComponent.h
@@ -15,12 +15,18 @@ class URemoteControlPreset;
  * e.g. duplicating an Actor with this component will result in the same properties exposed by Source Actor
  * to be automatically exposed by Duplicate Actor.
  */
-UCLASS()
-class REMOTECONTROLCOMPONENTS_API URemoteControlTrackerComponent : public UActorComponent
+UCLASS(MinimalAPI)
+class URemoteControlTrackerComponent : public UActorComponent
 {
 	GENERATED_BODY()
 
 public:
+	/** Gets the Current Preset this Tracker Component points to */
+	REMOTECONTROLCOMPONENTS_API URemoteControlPreset* GetCurrentPreset() const;
+
+	/** returns true if Tracker Component has at least one tracked/exposed property */
+	REMOTECONTROLCOMPONENTS_API bool HasTrackedProperties() const;
+
 	/** Add the specified property to the list of properties handled by this Tracker Component */
 	void AddTrackedProperty(const FRCFieldPathInfo& InFieldPathInfo, UObject* InOwnerObject);
 	
@@ -48,17 +54,12 @@ public:
 	 */
 	void WriteAllPropertyIdsToPreset();
 
-	/** returns true if Tracker Component has at least one tracked/exposed property */
-	bool HasTrackedProperties() const;
-
 	/** returns true if Tracker Component is tracking the property for the specified FieldPathInfo and Owner Object */
 	bool IsTrackingProperty(const FRCFieldPathInfo& InFieldPathInfo, UObject* InOwnerObject) const;
 	
 	/** Get the Actor owning this Tracker Component */
 	AActor* GetTrackedActor() const;
 
-	URemoteControlPreset* GetCurrentPreset() const;
-
 protected:
 	//~ Begin UActorComponent Interface
 	virtual void OnComponentCreated() override;
diff --git a/Engine/Plugins/Experimental/RemoteControlComponents/Source/RemoteControlComponents/Public/RemoteControlTrackerProperty.h b/Engine/Plugins/Experimental/RemoteControlComponents/Source/RemoteControlComponents/Public/RemoteControlTrackerProperty.h
--- a/Engine/Plugins/Experimental/RemoteControlComponents/Source/RemoteControlComponents/Public/RemoteControlTrackerProperty.h
+++ b/Engine/Plugins/Experimental/RemoteControlComponents/Source/RemoteControlComponents/Public/RemoteControlTrackerProperty.h
@@ -17,7 +17,7 @@ struct FRemoteControlTrackerProperty
 	GENERATED_BODY()
 
 public:
-	FRemoteControlTrackerProperty() {}
+	FRemoteControlTrackerProperty();
 
 	/**
 	 * Creates and initializes a Tracked Property
@@ -57,22 +57,11 @@ public:
 	 */
 	void MarkUnexposed();
 
-	bool operator==(const FRemoteControlTrackerProperty& Other) const
-	{
-		return
-			FieldPathInfo == Other.FieldPathInfo &&
-			OwnerObject == Other.OwnerObject;
-	}
-
-	bool operator!=(const FRemoteControlTrackerProperty& Other) const
-	{
-		return !(*this == Other);
-	}
-
-	friend uint32 GetTypeHash(const FRemoteControlTrackerProperty& InBroadcastControlId)
-	{
-		return HashCombineFast(InBroadcastControlId.FieldPathInfo.PathHash, GetTypeHash(InBroadcastControlId.OwnerObject));
-	}
+	bool operator==(const FRemoteControlTrackerProperty& Other) const;
+
+	bool operator!=(const FRemoteControlTrackerProperty& Other) const;
+
+	friend uint32 GetTypeHash(const FRemoteControlTrackerProperty& InBroadcastControlId);
 
 private:
 	/** Resolves the Tracked Property FieldPathInfo with its Owner */
diff --git a/Engine/Plugins/Experimental/RemoteControlComponents/Source/RemoteControlComponents/Public/Subsystems/RemoteControlComponentsSubsystem.h b/Engine/Plugins/Experimental/RemoteControlComponents/Source/RemoteControlComponents/Public/Subsystems/RemoteControlComponentsSubsystem.h
--- a/Engine/Plugins/Experimental/RemoteControlComponents/Source/RemoteControlComponents/Public/Subsystems/RemoteControlComponentsSubsystem.h
+++ b/Engine/Plugins/Experimental/RemoteControlComponents/Source/RemoteControlComponents/Public/Subsystems/RemoteControlComponentsSubsystem.h
@@ -2,8 +2,6 @@
 
 #pragma once
 
-#include "RemoteControlComponentsContext.h"
-#include "RemoteControlPreset.h"
 #include "Subsystems/EngineSubsystem.h"
 #include "Templates/SharedPointer.h"
 #include "UObject/ObjectKey.h"
@@ -13,23 +11,23 @@ class FSubsystemCollectionBase;
 class UObject;
 class URemoteControlPreset;
 class URemoteControlTrackerComponent;
-
 struct FRCFieldPathInfo;
+struct FRemoteControlComponentsContext;
 
-UCLASS()
-class REMOTECONTROLCOMPONENTS_API URemoteControlComponentsSubsystem : public UEngineSubsystem
+UCLASS(MinimalAPI)
+class URemoteControlComponentsSubsystem : public UEngineSubsystem
 {
 	GENERATED_BODY()
 
 public:
 	/** Get this subsystem instance */
-	static URemoteControlComponentsSubsystem* Get();
+	REMOTECONTROLCOMPONENTS_API static URemoteControlComponentsSubsystem* Get();
 
 	/** Register the specified Preset, along with the World it resides in. */
-	void RegisterPreset(URemoteControlPreset* InRemoteControlPreset);
+	REMOTECONTROLCOMPONENTS_API void RegisterPreset(URemoteControlPreset* InRemoteControlPreset);
 
 	/** Unregister the specified Preset */
-	void UnregisterPreset(URemoteControlPreset* InRemoteControlPreset);
+	REMOTECONTROLCOMPONENTS_API void UnregisterPreset(URemoteControlPreset* InRemoteControlPreset);
 
 	/** Marks the specified Actor as a Tracked Actor, so that the Subsystem knows about it */
 	void RegisterTrackedActor(AActor* InActor);
@@ -60,7 +58,7 @@ public:
 	/** Fires when the active Preset is updated */
 	FOnActivePresetChanged& OnActivePresetChanged() { return OnActivePresetChangedDelegate; }
 
-protected:
+private:
 	FOnActorRegisterStateUpdate OnTrackedActorRegisteredDelegate;	
 	FOnActorRegisterStateUpdate OnTrackedActorUnregisteredDelegate;
 	FOnActivePresetChanged OnActivePresetChangedDelegate;
diff --git a/Engine/Plugins/Experimental/RemoteControlComponents/Source/RemoteControlComponentsEditor/Private/RemoteControlComponentsEditorModule.cpp b/Engine/Plugins/Experimental/RemoteControlComponents/Source/RemoteControlComponentsEditor/Private/RemoteControlComponentsEditorModule.cpp
--- a/Engine/Plugins/Experimental/RemoteControlComponents/Source/RemoteControlComponentsEditor/Private/RemoteControlComponentsEditorModule.cpp
+++ b/Engine/Plugins/Experimental/RemoteControlComponents/Source/RemoteControlComponentsEditor/Private/RemoteControlComponentsEditorModule.cpp
@@ -1,5 +1,18 @@
 // Copyright Epic Games, Inc. All Rights Reserved.
 
 #include "Modules/ModuleManager.h"
+#include "Styles/RemoteControlComponentsEditorStyle.h"
 
-IMPLEMENT_MODULE(FDefaultModuleImpl, RemoteControlComponentsEditor)
+class FRemoteControlComponentsEditorModule : public IModuleInterface
+{
+public:
+	//~ Begin IModuleInterface
+	virtual void StartupModule() override
+	{
+		// Initialize the Style on Startup Module
+		FRemoteControlComponentsEditorStyle::Get();
+	}
+	//~ End IModuleInterface
+};
+
+IMPLEMENT_MODULE(FRemoteControlComponentsEditorModule, RemoteControlComponentsEditor)
diff --git a/Engine/Plugins/Experimental/RemoteControlComponents/Source/RemoteControlComponentsEditor/Private/Styles/RemoteControlComponentsEditorStyle.cpp b/Engine/Plugins/Experimental/RemoteControlComponents/Source/RemoteControlComponentsEditor/Private/Styles/RemoteControlComponentsEditorStyle.cpp
--- a/Engine/Plugins/Experimental/RemoteControlComponents/Source/RemoteControlComponentsEditor/Private/Styles/RemoteControlComponentsEditorStyle.cpp
+++ b/Engine/Plugins/Experimental/RemoteControlComponents/Source/RemoteControlComponentsEditor/Private/Styles/RemoteControlComponentsEditorStyle.cpp
@@ -13,11 +13,10 @@ FRemoteControlComponentsEditorStyle::FRemoteControlComponentsEditorStyle()
 	
 	const TSharedPtr<IPlugin> Plugin = IPluginManager::Get().FindPlugin(UE_PLUGIN_NAME);
 	check(Plugin.IsValid());
-	
+
 	ContentRootDir = FPaths::Combine(Plugin->GetBaseDir(), TEXT("Resources"));
 
-	Set("ClassIcon.RemoteControlComponentsComponent"    , new IMAGE_BRUSH_SVG("EditorIcons/RemoteControlComponentsDefault", Icon16x16));
-	Set("ClassIcon.RemoteControlTracker"                , new IMAGE_BRUSH_SVG("EditorIcons/RemoteControlTracker", Icon16x16));
+	Set("ClassIcon.RemoteControlTrackerComponent", new IMAGE_BRUSH_SVG("EditorIcons/RemoteControlTracker", Icon16x16));
 
 	FSlateStyleRegistry::RegisterSlateStyle(*this);
 }
diff --git a/Engine/Plugins/Experimental/RemoteControlComponents/Source/RemoteControlComponentsEditor/Private/Subsystems/RemoteControlComponentsEditorUtils.cpp b/Engine/Plugins/Experimental/RemoteControlComponents/Source/RemoteControlComponentsEditor/Private/Subsystems/RemoteControlComponentsEditorUtils.cpp
deleted file mode 100644
diff --git a/Engine/Plugins/Experimental/RemoteControlComponents/Source/RemoteControlComponentsEditor/Public/Subsystems/RemoteControlComponentsEditorUtils.h b/Engine/Plugins/Experimental/RemoteControlComponents/Source/RemoteControlComponentsEditor/Public/Subsystems/RemoteControlComponentsEditorUtils.h
deleted file mode 100644
diff --git a/Engine/Plugins/Experimental/RemoteControlComponents/Source/RemoteControlComponentsEditor/RemoteControlComponentsEditor.Build.cs b/Engine/Plugins/Experimental/RemoteControlComponents/Source/RemoteControlComponentsEditor/RemoteControlComponentsEditor.Build.cs
--- a/Engine/Plugins/Experimental/RemoteControlComponents/Source/RemoteControlComponentsEditor/RemoteControlComponentsEditor.Build.cs
+++ b/Engine/Plugins/Experimental/RemoteControlComponents/Source/RemoteControlComponentsEditor/RemoteControlComponentsEditor.Build.cs
@@ -8,16 +8,10 @@ public class RemoteControlComponentsEditor : ModuleRules
     {
         PCHUsage = ModuleRules.PCHUsageMode.UseExplicitOrSharedPCHs;
 
-        PublicDependencyModuleNames.AddRange(
-            new string[]
-            {
-                "Core"
-            }
-        );
-
         PrivateDependencyModuleNames.AddRange(
             new string[]
             {
+                "Core",
                 "CoreUObject",
                 "Engine",
                 "EditorSubsystem",
