Remote Control:
- Added static bool to return whether Guid Renewal is being allowed when duplicating
- Tracker Component to not refresh its tracker properties when Guid renewal is allowed, as it's analogous to duplicating things "as-is"

#changelist validated
#virtualized

diff --git a/Engine/Plugins/Experimental/RemoteControlComponents/Source/RemoteControlComponents/Private/RemoteControlTrackerComponent.cpp b/Engine/Plugins/Experimental/RemoteControlComponents/Source/RemoteControlComponents/Private/RemoteControlTrackerComponent.cpp
--- a/Engine/Plugins/Experimental/RemoteControlComponents/Source/RemoteControlComponents/Private/RemoteControlTrackerComponent.cpp
+++ b/Engine/Plugins/Experimental/RemoteControlComponents/Source/RemoteControlComponents/Private/RemoteControlTrackerComponent.cpp
@@ -143,7 +143,10 @@ void URemoteControlTrackerComponent::PostDuplicate(bool bInDuplicateForPIE)
 {
 	Super::PostDuplicate(bInDuplicateForPIE);
 
-	OnTrackerDuplicated();
+	if (!bInDuplicateForPIE)
+	{
+		OnTrackerDuplicated();
+	}
 }
 
 void URemoteControlTrackerComponent::PostEditImport()
@@ -201,6 +204,13 @@ void URemoteControlTrackerComponent::UnregisterTrackedActor() const
 
 void URemoteControlTrackerComponent::OnTrackerDuplicated()
 {
+	// Early out if Preset duplication is being blocked from the conventional Duplicate->Renew->RefreshTracker
+	// Not allowing preset guid renewal is analogous to duplicating the preset "as-is", and the tracker should not manipulate the preset
+	if (!FRCPresetGuidRenewGuard::IsAllowingPresetGuidRenewal())
+	{
+		return;
+	}
+
 	RefreshTracker();
 	WriteAllPropertyIdsToPreset();
 	RegisterPropertyIdChangeDelegate();
diff --git a/Engine/Plugins/VirtualProduction/RemoteControl/Source/RemoteControl/Private/RemoteControlPreset.cpp b/Engine/Plugins/VirtualProduction/RemoteControl/Source/RemoteControl/Private/RemoteControlPreset.cpp
--- a/Engine/Plugins/VirtualProduction/RemoteControl/Source/RemoteControl/Private/RemoteControlPreset.cpp
+++ b/Engine/Plugins/VirtualProduction/RemoteControl/Source/RemoteControl/Private/RemoteControlPreset.cpp
@@ -192,6 +192,11 @@ FRCPresetGuidRenewGuard::~FRCPresetGuidRenewGuard()
 	bAllowPresetGuidRenewal = bPreviousValue;
 }
 
+bool FRCPresetGuidRenewGuard::IsAllowingPresetGuidRenewal()
+{
+	return bAllowPresetGuidRenewal;
+}
+
 const TArray<FGuid>& FRemoteControlPresetGroup::GetFields() const
 {
 	return Fields;
diff --git a/Engine/Plugins/VirtualProduction/RemoteControl/Source/RemoteControl/Public/RemoteControlPreset.h b/Engine/Plugins/VirtualProduction/RemoteControl/Source/RemoteControl/Public/RemoteControlPreset.h
--- a/Engine/Plugins/VirtualProduction/RemoteControl/Source/RemoteControl/Public/RemoteControlPreset.h
+++ b/Engine/Plugins/VirtualProduction/RemoteControl/Source/RemoteControl/Public/RemoteControlPreset.h
@@ -137,7 +137,9 @@ struct REMOTECONTROL_API FRCPresetGuidRenewGuard
 {
 	FRCPresetGuidRenewGuard();
 	~FRCPresetGuidRenewGuard();
-		
+
+	static bool IsAllowingPresetGuidRenewal();
+
 private:
 	bool bPreviousValue;
 };
