[Remote Control Components] Remove "invalid" tracked properties from Tracker

- Remove properties referencing a no longer valid object from the tracked properties list

diff --git a/Engine/Plugins/Experimental/RemoteControlComponents/Source/RemoteControlComponents/Private/RemoteControlTrackerComponent.cpp b/Engine/Plugins/Experimental/RemoteControlComponents/Source/RemoteControlComponents/Private/RemoteControlTrackerComponent.cpp
--- a/Engine/Plugins/Experimental/RemoteControlComponents/Source/RemoteControlComponents/Private/RemoteControlTrackerComponent.cpp
+++ b/Engine/Plugins/Experimental/RemoteControlComponents/Source/RemoteControlComponents/Private/RemoteControlTrackerComponent.cpp
@@ -7,6 +7,7 @@
 #include "RemoteControlTrackerProperty.h"
 #include "RemoteControlPreset.h"
 #include "Subsystems/RemoteControlComponentsSubsystem.h"
+#include "UObject/ObjectSaveContext.h"
 
 URemoteControlPreset* URemoteControlTrackerComponent::GetCurrentPreset() const
 {
@@ -65,7 +66,7 @@ void URemoteControlTrackerComponent::UnexposeAllProperties()
 	{
 		TrackedProperty.Unexpose();
 	}
-	
+
 	UnregisterPropertyIdChangeDelegate();
 }
 
@@ -121,7 +122,7 @@ void URemoteControlTrackerComponent::OnComponentDestroyed(bool bInDestroyingHier
 	Super::OnComponentDestroyed(bInDestroyingHierarchy);
 
 	UnexposeAllProperties();
-	
+
 	if (AActor* TrackedActor = GetTrackedActor())
 	{
 		if (URemoteControlComponentsSubsystem* RemoteControlComponentsSubsystem = URemoteControlComponentsSubsystem::Get())
@@ -162,6 +163,13 @@ void URemoteControlTrackerComponent::PostLoad()
 	RegisterPropertyIdChangeDelegate();
 }
 
+void URemoteControlTrackerComponent::PreSave(FObjectPreSaveContext InSaveContext)
+{
+	Super::PreSave(InSaveContext);
+
+	CleanupProperties();
+}
+
 #if WITH_EDITOR
 void URemoteControlTrackerComponent::PostTransacted(const FTransactionObjectEvent& InTransactionEvent)
 {
@@ -169,8 +177,9 @@ void URemoteControlTrackerComponent::PostTransacted(const FTransactionObjectEven
 
 	if (InTransactionEvent.GetEventType() == ETransactionObjectEventType::UndoRedo)
 	{
+		CleanupProperties();
 		RefreshTracker();
-		
+
 		// In case the component was destroyed and created again with an Undo, let's try fixing potentially broken bindings
 		if (URemoteControlPreset* Preset = GetCurrentPreset())
 		{
@@ -181,7 +190,7 @@ void URemoteControlTrackerComponent::PostTransacted(const FTransactionObjectEven
 #endif
 
 void URemoteControlTrackerComponent::RegisterTrackedActor() const
-{	
+{
 	if (AActor* OuterActor = GetTrackedActor())
 	{
 		if (URemoteControlComponentsSubsystem* RemoteControlComponentsSubsystem = URemoteControlComponentsSubsystem::Get())
@@ -204,6 +213,8 @@ void URemoteControlTrackerComponent::UnregisterTrackedActor() const
 
 void URemoteControlTrackerComponent::OnTrackerDuplicated()
 {
+	CleanupProperties();
+
 	// Early out if Preset duplication is being blocked from the conventional Duplicate->Renew->RefreshTracker
 	// Not allowing preset guid renewal is analogous to duplicating the preset "as-is", and the tracker should not manipulate the preset
 	if (!FRCPresetGuidRenewGuard::IsAllowingPresetGuidRenewal())
@@ -236,6 +247,19 @@ void URemoteControlTrackerComponent::MarkPropertiesForRefresh()
 	}
 }
 
+void URemoteControlTrackerComponent::CleanupProperties()
+{
+	for (TArray<FRemoteControlTrackerProperty>::TIterator PropertyIt = TrackedProperties.CreateIterator(); PropertyIt; ++PropertyIt)
+	{
+		FRemoteControlTrackerProperty TrackedProperty(*PropertyIt);
+
+		if (!TrackedProperty.IsValid())
+		{
+			PropertyIt.RemoveCurrent();
+		}
+	}
+}
+
 void URemoteControlTrackerComponent::RegisterPropertyIdChangeDelegate()
 {
 	if (const URemoteControlPreset* Preset = GetCurrentPreset())
diff --git a/Engine/Plugins/Experimental/RemoteControlComponents/Source/RemoteControlComponents/Private/RemoteControlTrackerProperty.cpp b/Engine/Plugins/Experimental/RemoteControlComponents/Source/RemoteControlComponents/Private/RemoteControlTrackerProperty.cpp
--- a/Engine/Plugins/Experimental/RemoteControlComponents/Source/RemoteControlComponents/Private/RemoteControlTrackerProperty.cpp
+++ b/Engine/Plugins/Experimental/RemoteControlComponents/Source/RemoteControlComponents/Private/RemoteControlTrackerProperty.cpp
@@ -55,6 +55,11 @@ void FRemoteControlTrackerProperty::MarkUnexposed()
 	bIsExposed = false;
 }
 
+bool FRemoteControlTrackerProperty::IsValid() const
+{
+	return OwnerObject.IsValid();
+}
+
 bool FRemoteControlTrackerProperty::operator==(const FRemoteControlTrackerProperty& Other) const
 {
 	return FieldPathInfo == Other.FieldPathInfo && OwnerObject == Other.OwnerObject;
@@ -72,15 +77,15 @@ void FRemoteControlTrackerProperty::Expose(URemoteControlPreset* InRemoteControl
 		// Property is already exposed to specified preset, just return
 		return;
 	}
-	
-	if (!InRemoteControlPreset || !OwnerObject.IsValid())
+
+	if (!InRemoteControlPreset || !IsValid())
 	{
 		return;
 	}
-	
+
 	FRemoteControlComponentsUtils::ExposeProperty(InRemoteControlPreset, OwnerObject.Get(), FieldPathInfo);
 	CurrentPresetWeak = MakeWeakObjectPtr<URemoteControlPreset>(InRemoteControlPreset);
-	bIsExposed = true;	
+	bIsExposed = true;
 }
 
 void FRemoteControlTrackerProperty::Unexpose()
@@ -89,8 +94,8 @@ void FRemoteControlTrackerProperty::Unexpose()
 	{
 		return;
 	}
-	
-	if (!OwnerObject.IsValid())
+
+	if (!IsValid())
 	{
 		return;
 	}
@@ -105,6 +110,11 @@ void FRemoteControlTrackerProperty::Unexpose()
 
 void FRemoteControlTrackerProperty::ReadPropertyIdFromPreset()
 {
+	if (!IsValid())
+	{
+		return;
+	}
+
 	if (URemoteControlPreset* Preset = GetPreset())
 	{
 		const FName& Id = FRemoteControlComponentsUtils::GetExposedPropertyId(Preset, OwnerObject.Get(), FieldPathInfo);
@@ -117,7 +127,7 @@ void FRemoteControlTrackerProperty::ReadPropertyIdFromPreset()
 
 void FRemoteControlTrackerProperty::WritePropertyIdToPreset() const
 {
-	if (!IsExposed())
+	if (!IsExposed() || !IsValid())
 	{
 		return;
 	}
diff --git a/Engine/Plugins/Experimental/RemoteControlComponents/Source/RemoteControlComponents/Public/RemoteControlTrackerComponent.h b/Engine/Plugins/Experimental/RemoteControlComponents/Source/RemoteControlComponents/Public/RemoteControlTrackerComponent.h
--- a/Engine/Plugins/Experimental/RemoteControlComponents/Source/RemoteControlComponents/Public/RemoteControlTrackerComponent.h
+++ b/Engine/Plugins/Experimental/RemoteControlComponents/Source/RemoteControlComponents/Public/RemoteControlTrackerComponent.h
@@ -29,13 +29,13 @@ public:
 
 	/** Add the specified property to the list of properties handled by this Tracker Component */
 	void AddTrackedProperty(const FRCFieldPathInfo& InFieldPathInfo, UObject* InOwnerObject);
-	
+
 	/** Remove the specified property from the list of properties handled by this Tracker Component */
 	void RemoveTrackedProperty(const FRCFieldPathInfo& InFieldPathInfo, UObject* InOwnerObject);
 
 	/** Expose all Tracked Properties handled by this component to the current Remote Control Preset */
 	void ExposeAllProperties();
-	
+
 	/**
 	 * Unexpose all Tracked Properties handled by this component from the current Remote Control Preset.
 	 * It will also stop tracking all unexposed properties.
@@ -56,7 +56,7 @@ public:
 
 	/** returns true if Tracker Component is tracking the property for the specified FieldPathInfo and Owner Object */
 	bool IsTrackingProperty(const FRCFieldPathInfo& InFieldPathInfo, UObject* InOwnerObject) const;
-	
+
 	/** Get the Actor owning this Tracker Component */
 	AActor* GetTrackedActor() const;
 
@@ -65,12 +65,13 @@ protected:
 	virtual void OnComponentCreated() override;
 	virtual void OnComponentDestroyed(bool bInDestroyingHierarchy) override;
 	//~ End UActorComponent
-	
+
 	//~ Begin UObject Interface
 	virtual void PostInitProperties() override;
 	virtual void PostDuplicate(bool bInDuplicateForPIE) override;
 	virtual void PostEditImport() override;
 	virtual void PostLoad() override;
+	virtual void PreSave(FObjectPreSaveContext InSaveContext) override;
 #if WITH_EDITOR
 	virtual void PostTransacted(const FTransactionObjectEvent& InTransactionEvent) override;
 #endif
@@ -81,11 +82,12 @@ private:
 	void UnregisterTrackedActor() const;
 
 	void OnTrackerDuplicated();
-	
+
 	void RefreshTracker();
 	void RefreshExposedProperties();
 	void MarkPropertiesForRefresh();
-	
+	void CleanupProperties();
+
 	void RegisterPropertyIdChangeDelegate();
 	void UnregisterPropertyIdChangeDelegate() const;
 	void OnPropertyIdUpdated();
diff --git a/Engine/Plugins/Experimental/RemoteControlComponents/Source/RemoteControlComponents/Public/RemoteControlTrackerProperty.h b/Engine/Plugins/Experimental/RemoteControlComponents/Source/RemoteControlComponents/Public/RemoteControlTrackerProperty.h
--- a/Engine/Plugins/Experimental/RemoteControlComponents/Source/RemoteControlComponents/Public/RemoteControlTrackerProperty.h
+++ b/Engine/Plugins/Experimental/RemoteControlComponents/Source/RemoteControlComponents/Public/RemoteControlTrackerProperty.h
@@ -57,6 +57,11 @@ public:
 	 */
 	void MarkUnexposed();
 
+	/** Check if this property can still be exposed (internally verifies if Owner Object is valid) */
+	bool IsValid() const;
+
+	const FRCFieldPathInfo& GetFieldPathInfo() const { return FieldPathInfo; }
+
 	bool operator==(const FRemoteControlTrackerProperty& Other) const;
 
 	bool operator!=(const FRemoteControlTrackerProperty& Other) const;
